Add support for a Live Link Input Device.  Currently it only supports game pad devices on Win64 (mostly XInputInterface.cpp).  The LiveLink device establishes a thread to poll device data and send over live link.

diff --git a/Engine/Plugins/Animation/LiveLink/Source/LiveLinkHub/Private/LiveLinkHub.cpp b/Engine/Plugins/Animation/LiveLink/Source/LiveLinkHub/Private/LiveLinkHub.cpp
--- a/Engine/Plugins/Animation/LiveLink/Source/LiveLinkHub/Private/LiveLinkHub.cpp
+++ b/Engine/Plugins/Animation/LiveLink/Source/LiveLinkHub/Private/LiveLinkHub.cpp
@@ -64,6 +64,8 @@ void FLiveLinkHub::Initialize()
 	LiveLinkHubClient->OnLiveLinkSubjectRemoved().AddSP(this, &FLiveLinkHub::OnSubjectRemoved);
 
 	PlaybackController->Start();
+
+	GIsRunning = true;
 }
 
 FLiveLinkHub::~FLiveLinkHub()
diff --git a/Engine/Plugins/Runtime/Windows/XInputDevice/XInputDevice.uplugin b/Engine/Plugins/Runtime/Windows/XInputDevice/XInputDevice.uplugin
--- a/Engine/Plugins/Runtime/Windows/XInputDevice/XInputDevice.uplugin
+++ b/Engine/Plugins/Runtime/Windows/XInputDevice/XInputDevice.uplugin
@@ -10,6 +10,7 @@
 	"EnabledByDefault": true,
 	"IsBetaVersion": false,
 	"CanContainContent": false,
+    "SupportedPrograms": ["LiveLinkHub"],
 	"SupportedTargetPlatforms":
 	[
 		"Win64"
@@ -19,10 +20,11 @@
 		{
 			"Name": "XInputDevice",
 			"Type": "Runtime",
+            "ProgramAllowList": ["LiveLinkHub"],
 			"PlatformAllowList":
 			[
 				"Win64"
 			]
 		}
 	]
-}
\ No newline at end of file
+}
diff --git a/Engine/Plugins/VirtualProduction/LiveLinkInputDevice/LiveLinkInputDevice.uplugin b/Engine/Plugins/VirtualProduction/LiveLinkInputDevice/LiveLinkInputDevice.uplugin
new file mode 100644
--- /dev/null
+++ b/Engine/Plugins/VirtualProduction/LiveLinkInputDevice/LiveLinkInputDevice.uplugin
@@ -0,0 +1,38 @@
+{
+	"FileVersion": 1,
+	"Version": 1,
+	"VersionName": "1.0",
+	"FriendlyName": "LiveLinkInputDevice",
+	"Description": "Live Link plugin for Unreal Engine Input Devices, i.e. Game Controllers. It uses the InputDevice system to query values and share state over LiveLink.",
+	"Category": "Virtual Production",
+	"CreatedBy": "Epic Games, Inc.",
+	"CreatedByURL": "https://epicgames.com",
+	"DocsURL": "",
+	"MarketplaceURL": "",
+	"SupportURL": "",
+	"CanContainContent": true,
+	"IsBetaVersion": true,
+	"IsExperimentalVersion": false,
+	"Installed": false,
+	"EnabledByDefault": false,
+	"SupportedPrograms": ["LiveLinkHub"],
+	"SupportedTargetPlatforms": ["Win64"],
+	"Modules": [
+		{
+			"Name": "LiveLinkInputDevice",
+			"Type": "Runtime",
+			"LoadingPhase": "Default",
+			"ProgramAllowList": ["LiveLinkHub"],
+			"PlatformAllowList": [
+				"Win64"
+			]
+		}
+	],
+	"Plugins": [
+		{
+			"Name": "LiveLink",
+			"Enabled": true
+		}
+	]
+
+}
diff --git a/Engine/Plugins/VirtualProduction/LiveLinkInputDevice/Source/LiveLinkInputDevice/LiveLinkInputDevice.Build.cs b/Engine/Plugins/VirtualProduction/LiveLinkInputDevice/Source/LiveLinkInputDevice/LiveLinkInputDevice.Build.cs
new file mode 100644
--- /dev/null
+++ b/Engine/Plugins/VirtualProduction/LiveLinkInputDevice/Source/LiveLinkInputDevice/LiveLinkInputDevice.Build.cs
@@ -0,0 +1,33 @@
+// Copyright Epic Games, Inc. All Rights Reserved.
+
+using UnrealBuildTool;
+using System.IO;
+
+public class LiveLinkInputDevice : ModuleRules
+{
+	public LiveLinkInputDevice(ReadOnlyTargetRules Target) : base(Target)
+	{
+		PCHUsage = ModuleRules.PCHUsageMode.UseExplicitOrSharedPCHs;
+		
+		PublicDependencyModuleNames.AddRange(
+			new string[]
+			{
+				"Core",
+				"LiveLinkInterface",
+			}
+			);
+			
+		
+		PrivateDependencyModuleNames.AddRange(
+			new string[]
+			{
+				"ApplicationCore",
+				"CoreUObject",
+				"Engine",
+				"InputDevice",
+				"Slate",
+				"SlateCore"
+			}
+			);
+	}
+}
diff --git a/Engine/Plugins/VirtualProduction/LiveLinkInputDevice/Source/LiveLinkInputDevice/Private/LiveLinkInputDevice.cpp b/Engine/Plugins/VirtualProduction/LiveLinkInputDevice/Source/LiveLinkInputDevice/Private/LiveLinkInputDevice.cpp
new file mode 100644
--- /dev/null
+++ b/Engine/Plugins/VirtualProduction/LiveLinkInputDevice/Source/LiveLinkInputDevice/Private/LiveLinkInputDevice.cpp
@@ -0,0 +1,19 @@
+// Copyright Epic Games, Inc. All Rights Reserved.
+
+#include "LiveLinkInputDevice.h"
+
+#define LOCTEXT_NAMESPACE "LiveLinkInputDeviceModule"
+
+DEFINE_LOG_CATEGORY(LogLiveLinkInputDevice);
+
+void FLiveLinkInputDeviceModule::StartupModule()
+{
+}
+
+void FLiveLinkInputDeviceModule::ShutdownModule()
+{
+}
+
+#undef LOCTEXT_NAMESPACE
+	
+IMPLEMENT_MODULE(FLiveLinkInputDeviceModule, LiveLinkInputDevice)
\ No newline at end of file
diff --git a/Engine/Plugins/VirtualProduction/LiveLinkInputDevice/Source/LiveLinkInputDevice/Private/LiveLinkInputDeviceMessageHandler.cpp b/Engine/Plugins/VirtualProduction/LiveLinkInputDevice/Source/LiveLinkInputDevice/Private/LiveLinkInputDeviceMessageHandler.cpp
new file mode 100644
--- /dev/null
+++ b/Engine/Plugins/VirtualProduction/LiveLinkInputDevice/Source/LiveLinkInputDevice/Private/LiveLinkInputDeviceMessageHandler.cpp
@@ -0,0 +1,43 @@
+// Copyright Epic Games, Inc. All Rights Reserved.
+
+#include "LiveLinkInputDeviceMessageHandler.h"
+#include "LiveLinkInputDeviceTypes.h"
+#include "Misc/CoreMiscDefines.h"
+
+
+bool FLiveLinkInputDeviceMessageHandler::OnControllerAnalog(FGamepadKeyNames::Type KeyName, FPlatformUserId PlatformUserId, FInputDeviceId InputDeviceId, float AnalogValue)
+{
+	FLiveLinkGamepadInputDeviceFrameData& InputDeviceFrameData = CurrentFrameDataValues.FindOrAdd(InputDeviceId);
+	InputDeviceFrameData.ApplyValueFromKey(KeyName, AnalogValue);
+	return true;
+}
+
+bool FLiveLinkInputDeviceMessageHandler::OnControllerButtonPressed(FGamepadKeyNames::Type KeyName, FPlatformUserId PlatformUserId, FInputDeviceId InputDeviceId, bool IsRepeat)
+{
+	FLiveLinkGamepadInputDeviceFrameData& InputDeviceFrameData = CurrentFrameDataValues.FindOrAdd(InputDeviceId);
+	InputDeviceFrameData.ApplyValueFromKey(KeyName, 1);
+	return true;
+}
+
+bool FLiveLinkInputDeviceMessageHandler::OnControllerButtonReleased(FGamepadKeyNames::Type KeyName, FPlatformUserId PlatformUserId, FInputDeviceId InputDeviceId, bool IsRepeat)
+{
+	FLiveLinkGamepadInputDeviceFrameData& InputDeviceFrameData = CurrentFrameDataValues.FindOrAdd(InputDeviceId);
+	InputDeviceFrameData.ApplyValueFromKey(KeyName, 0);
+	return true;
+}
+
+FLiveLinkGamepadInputDeviceFrameData FLiveLinkInputDeviceMessageHandler::GetLatestValue(FInputDeviceId InDeviceId /*unused*/) const
+{
+	for (const TPair<FInputDeviceId, FLiveLinkGamepadInputDeviceFrameData>& Device : CurrentFrameDataValues)
+	{
+		return Device.Value;
+	}
+	return {};
+}
+
+TSet<FInputDeviceId> FLiveLinkInputDeviceMessageHandler::GetDeviceIds() const
+{
+	TSet<FInputDeviceId> Devices;
+	CurrentFrameDataValues.GetKeys(Devices);
+	return Devices;
+}
diff --git a/Engine/Plugins/VirtualProduction/LiveLinkInputDevice/Source/LiveLinkInputDevice/Private/LiveLinkInputDeviceMessageHandler.h b/Engine/Plugins/VirtualProduction/LiveLinkInputDevice/Source/LiveLinkInputDevice/Private/LiveLinkInputDeviceMessageHandler.h
new file mode 100644
--- /dev/null
+++ b/Engine/Plugins/VirtualProduction/LiveLinkInputDevice/Source/LiveLinkInputDevice/Private/LiveLinkInputDeviceMessageHandler.h
@@ -0,0 +1,27 @@
+// Copyright Epic Games, Inc. All Rights Reserved.
+
+#pragma once
+
+#include "GenericPlatform/GenericApplicationMessageHandler.h"
+#include "LiveLinkInputDeviceTypes.h"
+#include "Containers/Map.h"
+#include "Misc/CoreMiscDefines.h"
+
+class FLiveLinkInputDeviceMessageHandler : public FGenericApplicationMessageHandler
+{
+public:
+	virtual ~FLiveLinkInputDeviceMessageHandler() = default;
+
+	virtual bool OnControllerAnalog(FGamepadKeyNames::Type KeyName, FPlatformUserId PlatformUserId, FInputDeviceId InputDeviceId, float AnalogValue) override;
+	virtual bool OnControllerButtonPressed(FGamepadKeyNames::Type KeyName, FPlatformUserId PlatformUserId, FInputDeviceId InputDeviceId, bool IsRepeat) override;
+	virtual bool OnControllerButtonReleased(FGamepadKeyNames::Type KeyName, FPlatformUserId PlatformUserId, FInputDeviceId InputDeviceId, bool IsRepeat) override;
+
+	/** For the given device id, which maps the the controller id, get the latest frame values for that device. */
+	FLiveLinkGamepadInputDeviceFrameData GetLatestValue(FInputDeviceId InDeviceId) const;
+
+	/** Return the current known list of devices */
+	TSet<FInputDeviceId> GetDeviceIds() const;
+
+private:
+	TMap<FInputDeviceId, FLiveLinkGamepadInputDeviceFrameData> CurrentFrameDataValues;
+};
diff --git a/Engine/Plugins/VirtualProduction/LiveLinkInputDevice/Source/LiveLinkInputDevice/Private/LiveLinkInputDeviceRole.cpp b/Engine/Plugins/VirtualProduction/LiveLinkInputDevice/Source/LiveLinkInputDevice/Private/LiveLinkInputDeviceRole.cpp
new file mode 100644
--- /dev/null
+++ b/Engine/Plugins/VirtualProduction/LiveLinkInputDevice/Source/LiveLinkInputDevice/Private/LiveLinkInputDeviceRole.cpp
@@ -0,0 +1,46 @@
+// Copyright Epic Games, Inc. All Rights Reserved.
+
+
+#include "LiveLinkInputDeviceRole.h"
+#include "LiveLinkInputDeviceTypes.h"
+
+#define LOCTEXT_NAMESPACE "LiveLinkInputDeviceRole"
+
+UScriptStruct* ULiveLinkInputDeviceRole::GetStaticDataStruct() const
+{
+	return FLiveLinkGamepadInputDeviceStaticData::StaticStruct();
+}
+
+UScriptStruct* ULiveLinkInputDeviceRole::GetFrameDataStruct() const
+{
+	return FLiveLinkGamepadInputDeviceFrameData::StaticStruct();
+}
+
+UScriptStruct* ULiveLinkInputDeviceRole::GetBlueprintDataStruct() const
+{
+	return FLiveLinkGamepadInputDeviceBlueprintData::StaticStruct();
+}
+
+bool ULiveLinkInputDeviceRole::InitializeBlueprintData(const FLiveLinkSubjectFrameData& InSourceData, FLiveLinkBlueprintDataStruct& OutBlueprintData) const
+{
+	bool bSuccess = false;
+
+	FLiveLinkGamepadInputDeviceBlueprintData* BlueprintData = OutBlueprintData.Cast<FLiveLinkGamepadInputDeviceBlueprintData>();
+	const FLiveLinkGamepadInputDeviceStaticData* StaticData = InSourceData.StaticData.Cast<FLiveLinkGamepadInputDeviceStaticData>();
+	const FLiveLinkGamepadInputDeviceFrameData* FrameData = InSourceData.FrameData.Cast<FLiveLinkGamepadInputDeviceFrameData>();
+	if (BlueprintData && StaticData && FrameData)
+	{
+		GetStaticDataStruct()->CopyScriptStruct(&BlueprintData->StaticData, StaticData);
+		GetFrameDataStruct()->CopyScriptStruct(&BlueprintData->FrameData, FrameData);
+		bSuccess = true;
+	}
+
+	return bSuccess;
+}
+
+FText ULiveLinkInputDeviceRole::GetDisplayName() const
+{
+	return LOCTEXT("LiveLinkInputDevice", "Live Link Input Device");
+}
+
+#undef LOCTEXT_NAMESPACE
diff --git a/Engine/Plugins/VirtualProduction/LiveLinkInputDevice/Source/LiveLinkInputDevice/Private/LiveLinkInputDeviceSource.cpp b/Engine/Plugins/VirtualProduction/LiveLinkInputDevice/Source/LiveLinkInputDevice/Private/LiveLinkInputDeviceSource.cpp
new file mode 100644
--- /dev/null
+++ b/Engine/Plugins/VirtualProduction/LiveLinkInputDevice/Source/LiveLinkInputDevice/Private/LiveLinkInputDeviceSource.cpp
@@ -0,0 +1,285 @@
+// Copyright Epic Games, Inc. All Rights Reserved.
+
+#include "LiveLinkInputDeviceSource.h"
+
+#include "Async/Async.h"
+#include "CoreFwd.h"
+#include "Engine/Engine.h"
+#include "Features/IModularFeature.h"
+#include "GenericPlatform/GenericApplicationMessageHandler.h"
+#include "GenericPlatform/IInputInterface.h"
+#include "IInputDeviceModule.h"
+#include "InputCoreTypes.h"
+#include "Misc/CoreDelegates.h"
+#include "Roles/LiveLinkTransformRole.h"
+
+#include "ILiveLinkClient.h"
+#include "LiveLinkInputDevice.h"
+#include "LiveLinkInputDeviceRole.h"
+#include "LiveLinkInputDeviceSourceSettings.h"
+#include "LiveLinkInputDeviceTypes.h"
+#include "LiveLinkTypes.h"
+
+#include "LiveLinkInputDeviceMessageHandler.h"
+
+#define LOCTEXT_NAMESPACE "LiveLinkInputDeviceSource"
+
+FLiveLinkInputDeviceSource::FLiveLinkInputDeviceSource(const FLiveLinkInputDeviceConnectionSettings& ConnectionSettings)
+{
+	UpdateGamepadState(false);
+	SourceType = LOCTEXT("SourceType_InputDevice", "Input Device");
+	SourceMachineName = LOCTEXT("LocalMachine", "Local");
+
+	DeferredStartDelegateHandle = FCoreDelegates::OnEndFrame.AddRaw(this, &FLiveLinkInputDeviceSource::Start);
+}
+
+FLiveLinkInputDeviceSource::~FLiveLinkInputDeviceSource()
+{
+	if (DeferredStartDelegateHandle.IsValid())
+	{
+		FCoreDelegates::OnEndFrame.Remove(DeferredStartDelegateHandle);
+	}
+}
+
+void FLiveLinkInputDeviceSource::ReceiveClient(ILiveLinkClient* InClient, FGuid InSourceGuid)
+{
+	Client = InClient;
+	SourceGuid = InSourceGuid;
+}
+
+bool FLiveLinkInputDeviceSource::IsSourceStillValid() const
+{
+	const bool bIsSourceValid = !Stopping && (Thread != nullptr);
+	return bIsSourceValid;
+}
+
+bool FLiveLinkInputDeviceSource::RequestSourceShutdown()
+{
+	Stop();
+
+	return true;
+}
+
+// FRunnable interface
+void FLiveLinkInputDeviceSource::Start()
+{
+	check(DeferredStartDelegateHandle.IsValid());
+
+	FCoreDelegates::OnEndFrame.Remove(DeferredStartDelegateHandle);
+	DeferredStartDelegateHandle.Reset();
+
+	ThreadName = "LiveLinkInputDevice Receiver";
+	ThreadName.AppendInt(FAsyncThreadIndex::GetNext());
+
+	Thread.Reset(FRunnableThread::Create(this, *ThreadName, 128 * 1024, TPri_AboveNormal, FPlatformAffinity::GetPoolThreadMask()));
+}
+
+void FLiveLinkInputDeviceSource::Stop()
+{
+	Stopping = true;
+}
+
+void FLiveLinkInputDeviceSource::InitializeSettings(ULiveLinkSourceSettings* Settings)
+{
+	Settings->Mode = ELiveLinkSourceMode::Latest;
+}
+
+void FLiveLinkInputDeviceSource::LoadInputPlugins()
+{
+	MessageHandler = MakeShareable(new FLiveLinkInputDeviceMessageHandler());
+	InputDevices.Reset();
+
+	TArray<IInputDeviceModule*> PluginImplementations = IModularFeatures::Get().GetModularFeatureImplementations<IInputDeviceModule>( IInputDeviceModule::GetModularFeatureName() );
+	for( auto InputPluginIt = PluginImplementations.CreateIterator(); InputPluginIt; ++InputPluginIt )
+	{
+		TSharedPtr<IInputDevice> Device = (*InputPluginIt)->CreateInputDevice(
+			MessageHandler.ToSharedRef(),
+			{
+			.bInitAsPrimaryDevice = false,
+		    });
+		if (Device)
+		{
+			InputDevices.Add(Device);
+		}
+	}
+
+	if (InputDevices.Num() > 0)
+	{
+		bHasLoadedInputPlugins = true;
+	}
+}
+
+void FLiveLinkInputDeviceSource::UpdateGamepadState(const bool bInState)
+{
+	if (bInState)
+	{
+		SourceStatus = LOCTEXT("SourceStatus_Receiving", "Receiving");
+	}
+	else
+	{
+		SourceStatus = LOCTEXT("SourceStatus_Waiting", "Waiting for a device.");
+	}
+	bIsGamepadConnected = bInState;
+}
+
+bool FLiveLinkInputDeviceSource::TickInputDevices(const double InDeltaTime)
+{
+	static thread_local double ElapsedUpdate = 0;
+	bool bGamepadState = false;
+	for( int Index = 0; Index < InputDevices.Num(); Index++ )
+	{
+		InputDevices[Index]->Tick( InDeltaTime );
+		InputDevices[Index]->SendControllerEvents();
+		const bool bDeviceGamepadAttached = InputDevices[Index]->IsGamepadAttached();
+		// If elapsed time is greater than 1 second.  Force a device update so that we can catch new devices
+		// added to the system.  This normally happens by the WindowsApplication.cpp but we are outside of the
+		// process message loop so we have to periodically apply it here. We only do this here if we don't have
+		// any devices attached as there is a cost for making this query that can cause noticeble hitches.
+		//
+		if (!bDeviceGamepadAttached && ElapsedUpdate > 1.0)
+		{
+			static const FInputDeviceProperty RequestUpdateProp(TEXT("Request_Device_Update"));
+			InputDevices[Index]->SetDeviceProperty(-1, &RequestUpdateProp);
+		}
+		bGamepadState = bGamepadState || bDeviceGamepadAttached;
+	}
+
+	ElapsedUpdate = ElapsedUpdate > 1 ? 0 : ElapsedUpdate + InDeltaTime;
+
+	// Update our device for Live Link UI.
+	//
+	if (bIsGamepadConnected != bGamepadState)
+	{
+		UpdateGamepadState(bGamepadState);
+	}
+
+	return bGamepadState;
+}
+
+double FLiveLinkInputDeviceSource::DevicePollWaitTime() const
+{
+	const double UpdateRateInHz = LocalUpdateRateInHz > 0 ? (double) LocalUpdateRateInHz : 60.0;
+	return 1.0 / UpdateRateInHz;
+}
+
+uint32 FLiveLinkInputDeviceSource::Run()
+{
+	// The amount of time we sleep on one iteration of this loop.  This is not const because we only update at the desired
+	// update rate if we have input devices attached.  If no device is attached then we need to "wait" for a period of time
+	// before reruning our loop.
+	//
+	double SleepDeltaTime = 1.0 / (double)LocalUpdateRateInHz;
+
+	// Our wait time if no devices are initially connected.
+	//
+	const double WaitForDeviceTime = 1.0; // 1 Second.
+
+	double LastUpdateTime = FPlatformTime::Seconds();
+	while (!Stopping)
+	{
+		// Make sure we do not exceed our local update rate.
+		ON_SCOPE_EXIT{
+			FPlatformProcess::Sleep(SleepDeltaTime);
+		};
+
+		if (!GIsRunning)
+		{
+			continue;
+		}
+
+		if (!bHasLoadedInputPlugins)
+		{
+			LoadInputPlugins();
+		}
+
+		// Poll externally-implemented devices
+		//
+		const double CurrentTime = FPlatformTime::Seconds();
+		TickInputDevices(CurrentTime - LastUpdateTime);;
+		LastUpdateTime = CurrentTime;
+
+		TSet<FInputDeviceId> Devices = MessageHandler->GetDeviceIds();
+
+		if (Devices.Num() == 0)
+		{
+			// If no devices then we throttle back our update loop and wait for a device to come online.
+			SleepDeltaTime = WaitForDeviceTime;
+			continue;
+		}
+
+		for (FInputDeviceId Id : Devices)
+		{
+			if (TOptional<FLiveLinkStaticDataStruct> StaticData = PollGamepadStaticData(Id))
+			{
+				Send(Id, MoveTemp(*StaticData));
+			}
+			else if (TOptional<FLiveLinkFrameDataStruct> FrameData = PollGamepadFrameData(Id))
+			{
+				Send(Id, MoveTemp(*FrameData));
+			}
+		}
+
+		SleepDeltaTime = DevicePollWaitTime();
+	}
+	
+	return 0;
+}
+
+TOptional<FLiveLinkStaticDataStruct> FLiveLinkInputDeviceSource::PollGamepadStaticData(FInputDeviceId InDeviceId)
+{
+	if (Client)
+	{
+		if (!EncounteredSubjects.Contains(InDeviceId))
+		{
+			EncounteredSubjects.Add(InDeviceId);
+			FLiveLinkStaticDataStruct StaticData(FLiveLinkGamepadInputDeviceStaticData::StaticStruct());
+			return StaticData;
+		}
+	}
+
+	return {};
+}
+
+TOptional<FLiveLinkFrameDataStruct> FLiveLinkInputDeviceSource::PollGamepadFrameData(FInputDeviceId InDeviceId)
+{
+	if (Client)
+	{
+		FLiveLinkFrameDataStruct BaseFrameData(FLiveLinkGamepadInputDeviceFrameData::StaticStruct());
+		FLiveLinkGamepadInputDeviceFrameData* FrameData = BaseFrameData.Cast<FLiveLinkGamepadInputDeviceFrameData>();
+
+		*FrameData = MessageHandler->GetLatestValue({});
+		FrameData->WorldTime = FPlatformTime::Seconds();
+
+		return BaseFrameData;
+	}
+	return {};
+}
+
+void FLiveLinkInputDeviceSource::Send(FInputDeviceId DeviceId, FLiveLinkStaticDataStruct StaticDataToSend)
+{
+	if (Stopping || Client == nullptr)
+	{
+		return;
+	}
+	FNameBuilder Builder("Gamepad_");
+	Builder.Append(FString::FromInt(DeviceId.GetId()));
+
+	FName SubjectName = *Builder;
+	Client->PushSubjectStaticData_AnyThread({ SourceGuid, SubjectName }, ULiveLinkInputDeviceRole::StaticClass(), MoveTemp(StaticDataToSend));
+}
+
+void FLiveLinkInputDeviceSource::Send(FInputDeviceId DeviceId, FLiveLinkFrameDataStruct FrameDataToSend)
+{
+	if (Stopping || Client == nullptr)
+	{
+		return;
+	}
+	FNameBuilder Builder("Gamepad_");
+	Builder.Append(FString::FromInt(DeviceId.GetId()));
+	FName SubjectName = *Builder;
+	Client->PushSubjectFrameData_AnyThread({ SourceGuid, SubjectName }, MoveTemp(FrameDataToSend));
+}
+
+
+
+#undef LOCTEXT_NAMESPACE
diff --git a/Engine/Plugins/VirtualProduction/LiveLinkInputDevice/Source/LiveLinkInputDevice/Private/LiveLinkInputDeviceSourceFactory.cpp b/Engine/Plugins/VirtualProduction/LiveLinkInputDevice/Source/LiveLinkInputDevice/Private/LiveLinkInputDeviceSourceFactory.cpp
new file mode 100644
--- /dev/null
+++ b/Engine/Plugins/VirtualProduction/LiveLinkInputDevice/Source/LiveLinkInputDevice/Private/LiveLinkInputDeviceSourceFactory.cpp
@@ -0,0 +1,48 @@
+// Copyright Epic Games, Inc. All Rights Reserved.
+
+#include "LiveLinkInputDeviceSourceFactory.h"
+
+#include "Features/IModularFeatures.h"
+#include "ILiveLinkClient.h"
+#include "LiveLinkInputDeviceSource.h"
+#include "LiveLinkInputDeviceSourceSettings.h"
+#include "SLiveLinkInputDeviceSourceFactory.h"
+
+#define LOCTEXT_NAMESPACE "LiveLinkInputDeviceSourceFactory"
+
+FText ULiveLinkInputDeviceSourceFactory::GetSourceDisplayName() const
+{
+	return LOCTEXT("SourceDisplayName", "LiveLinkInputDevice Source");	
+}
+
+FText ULiveLinkInputDeviceSourceFactory::GetSourceTooltip() const
+{
+	return LOCTEXT("SourceTooltip", "Allows creation of LiveLink sources using the Unreal Engine input device system. Only gamepad controllers are currently supported.");
+}
+
+TSharedPtr<SWidget> ULiveLinkInputDeviceSourceFactory::BuildCreationPanel(FOnLiveLinkSourceCreated InOnLiveLinkSourceCreated) const
+{
+	return SNew(SLiveLinkInputDeviceSourceFactory)
+		.OnConnectionSettingsAccepted(FOnLiveLinkInputDeviceConnectionSettingsAccepted::CreateUObject(this, &ULiveLinkInputDeviceSourceFactory::CreateSourceFromSettings, InOnLiveLinkSourceCreated));
+}
+
+TSharedPtr<ILiveLinkSource> ULiveLinkInputDeviceSourceFactory::CreateSource(const FString& ConnectionString) const
+{
+	FLiveLinkInputDeviceConnectionSettings ConnectionSettings;
+	if (!ConnectionString.IsEmpty())
+	{
+		FLiveLinkInputDeviceConnectionSettings::StaticStruct()->ImportText(*ConnectionString, &ConnectionSettings, nullptr, PPF_None, GLog, TEXT("ULiveLinkInputDeviceSourceFactory"));
+	}
+	return MakeShared<FLiveLinkInputDeviceSource>(ConnectionSettings);
+}
+
+void ULiveLinkInputDeviceSourceFactory::CreateSourceFromSettings(FLiveLinkInputDeviceConnectionSettings InConnectionSettings, FOnLiveLinkSourceCreated OnSourceCreated) const
+{
+	FString ConnectionString;
+	FLiveLinkInputDeviceConnectionSettings::StaticStruct()->ExportText(ConnectionString, &InConnectionSettings, nullptr, nullptr, PPF_None, nullptr);
+
+	TSharedPtr<FLiveLinkInputDeviceSource> SharedPtr = MakeShared<FLiveLinkInputDeviceSource>(InConnectionSettings);
+	OnSourceCreated.ExecuteIfBound(SharedPtr, MoveTemp(ConnectionString));
+}
+
+#undef LOCTEXT_NAMESPACE
diff --git a/Engine/Plugins/VirtualProduction/LiveLinkInputDevice/Source/LiveLinkInputDevice/Private/LiveLinkInputDeviceTypes.cpp b/Engine/Plugins/VirtualProduction/LiveLinkInputDevice/Source/LiveLinkInputDevice/Private/LiveLinkInputDeviceTypes.cpp
new file mode 100644
--- /dev/null
+++ b/Engine/Plugins/VirtualProduction/LiveLinkInputDevice/Source/LiveLinkInputDevice/Private/LiveLinkInputDeviceTypes.cpp
@@ -0,0 +1,155 @@
+// Copyright Epic Games, Inc. All Rights Reserved.
+
+#include "LiveLinkInputDeviceTypes.h"
+#include "GenericPlatform/GenericApplicationMessageHandler.h"
+
+void SetValueWithKey(FName Key, float InValue, FLiveLinkGamepadInputDeviceFrameData& OutData)
+{
+	if (Key == FGamepadKeyNames::LeftAnalogX)
+	{
+		OutData.LeftAnalogX = InValue;
+	}
+	else
+	if (Key == FGamepadKeyNames::LeftAnalogY)
+	{
+		OutData.LeftAnalogY = InValue;
+	}
+	else
+	if (Key == FGamepadKeyNames::RightAnalogX)
+	{
+		OutData.RightAnalogX = InValue;
+	}
+	else
+	if (Key == FGamepadKeyNames::RightAnalogY)
+	{
+		OutData.RightAnalogY = InValue;
+	}
+	else
+	if (Key == FGamepadKeyNames::LeftTriggerAnalog)
+	{
+		OutData.LeftTriggerAnalog = InValue;
+	}
+	else
+	if (Key == FGamepadKeyNames::RightTriggerAnalog)
+	{
+		OutData.RightTriggerAnalog = InValue;
+	}
+	else
+	if (Key == FGamepadKeyNames::LeftThumb)
+	{
+		OutData.LeftThumb = InValue;
+	}
+	else
+	if (Key == FGamepadKeyNames::RightThumb)
+	{
+		OutData.RightThumb = InValue;
+	}
+	else
+	if (Key == FGamepadKeyNames::SpecialLeft)
+	{
+		OutData.SpecialLeft = InValue;
+	}
+	else
+	if (Key == FGamepadKeyNames::SpecialLeft_X)
+	{
+		OutData.SpecialLeft_X = InValue;
+	}
+	else if (Key == FGamepadKeyNames::SpecialLeft_Y)
+	{
+		OutData.SpecialLeft_Y = InValue;
+	}
+	else if (Key == FGamepadKeyNames::SpecialRight)
+	{
+		OutData.SpecialRight = InValue;
+	}
+	else if (Key == FGamepadKeyNames::FaceButtonBottom)
+	{
+		OutData.FaceButtonBottom = InValue;
+	}
+	else if (Key == FGamepadKeyNames::FaceButtonRight)
+	{
+		OutData.FaceButtonRight = InValue;
+	}
+	else if (Key == FGamepadKeyNames::FaceButtonLeft)
+	{
+		OutData.FaceButtonLeft = InValue;
+	}
+	else if (Key == FGamepadKeyNames::FaceButtonTop)
+	{
+		OutData.FaceButtonTop = InValue;
+	}
+	else if (Key == FGamepadKeyNames::LeftShoulder)
+	{
+		OutData.LeftShoulder = InValue;
+	}
+	else if (Key == FGamepadKeyNames::RightShoulder)
+	{
+		OutData.RightShoulder = InValue;
+	}
+	else if (Key == FGamepadKeyNames::LeftTriggerThreshold)
+	{
+		OutData.LeftTriggerThreshold = InValue;
+	}
+	else if (Key == FGamepadKeyNames::RightTriggerThreshold)
+	{
+		OutData.RightTriggerThreshold = InValue;
+	}
+	else if (Key == FGamepadKeyNames::DPadUp)
+	{
+		OutData.DPadUp = InValue;
+	}
+	else if (Key == FGamepadKeyNames::DPadDown)
+	{
+		OutData.DPadDown = InValue;
+	}
+	else if (Key == FGamepadKeyNames::DPadRight)
+	{
+		OutData.DPadRight = InValue;
+	}
+	else if (Key == FGamepadKeyNames::DPadLeft)
+	{
+		OutData.DPadLeft = InValue;
+	}
+	else if (Key == FGamepadKeyNames::LeftStickUp)
+	{
+		OutData.LeftStickUp = InValue;
+	}
+	else if (Key == FGamepadKeyNames::LeftStickDown)
+	{
+		OutData.LeftStickDown = InValue;
+	}
+	else if (Key == FGamepadKeyNames::LeftStickRight)
+	{
+		OutData.LeftStickRight = InValue;
+	}
+	else if (Key == FGamepadKeyNames::LeftStickLeft)
+	{
+		OutData.LeftStickLeft = InValue;
+	}
+	else if (Key == FGamepadKeyNames::RightStickUp)
+	{
+		OutData.RightStickUp = InValue;
+	}
+	else if (Key == FGamepadKeyNames::RightStickDown)
+	{
+		OutData.RightStickDown = InValue;
+	}
+	else if (Key == FGamepadKeyNames::RightStickRight)
+	{
+		OutData.RightStickRight = InValue;
+	}
+	else if (Key == FGamepadKeyNames::RightStickLeft)
+	{
+		OutData.RightStickLeft = InValue;
+	}
+	else
+	{
+		ensureMsgf(false, TEXT("Unknown key value %s"), *Key.ToString());
+	}
+
+}
+
+void FLiveLinkGamepadInputDeviceFrameData::ApplyValueFromKey(FName Key, float InValue)
+{
+	SetValueWithKey(Key, InValue, *this);
+}
diff --git a/Engine/Plugins/VirtualProduction/LiveLinkInputDevice/Source/LiveLinkInputDevice/Private/SLiveLinkInputDeviceSourceFactory.cpp b/Engine/Plugins/VirtualProduction/LiveLinkInputDevice/Source/LiveLinkInputDevice/Private/SLiveLinkInputDeviceSourceFactory.cpp
new file mode 100644
--- /dev/null
+++ b/Engine/Plugins/VirtualProduction/LiveLinkInputDevice/Source/LiveLinkInputDevice/Private/SLiveLinkInputDeviceSourceFactory.cpp
@@ -0,0 +1,62 @@
+// Copyright Epic Games, Inc. All Rights Reserved.
+
+#include "SLiveLinkInputDeviceSourceFactory.h"
+#include "LiveLinkInputDevice.h"
+#include "LiveLinkInputDeviceSourceSettings.h"
+
+#include "Widgets/Input/SButton.h"
+#include "Widgets/SBoxPanel.h"
+
+#if WITH_EDITOR
+#include "DetailLayoutBuilder.h"
+#endif //WITH_EDITOR
+
+#define LOCTEXT_NAMESPACE "SLiveLinkInputDeviceSourceFactory"
+
+void SLiveLinkInputDeviceSourceFactory::Construct(const FArguments& Args)
+{
+#if WITH_EDITOR
+	OnConnectionSettingsAccepted = Args._OnConnectionSettingsAccepted;
+
+	FStructureDetailsViewArgs StructureViewArgs;
+	FDetailsViewArgs DetailArgs;
+	DetailArgs.bAllowSearch = false;
+	DetailArgs.bShowScrollBar = false;
+
+	FPropertyEditorModule& PropertyEditor = FModuleManager::Get().LoadModuleChecked<FPropertyEditorModule>(TEXT("PropertyEditor"));
+
+	StructOnScope = MakeShared<FStructOnScope>(FLiveLinkInputDeviceConnectionSettings::StaticStruct());
+	CastChecked<UScriptStruct>(StructOnScope->GetStruct())->CopyScriptStruct(StructOnScope->GetStructMemory(), &ConnectionSettings);
+	StructureDetailsView = PropertyEditor.CreateStructureDetailView(DetailArgs, StructureViewArgs, StructOnScope);
+
+	ChildSlot
+	[
+		SNew(SVerticalBox)
+		+SVerticalBox::Slot()
+		.FillHeight(1.f)
+		[
+			StructureDetailsView->GetWidget().ToSharedRef()
+		]
+		+ SVerticalBox::Slot()
+		.HAlign(HAlign_Right)
+		.AutoHeight()
+		[
+			SNew(SButton)
+			.OnClicked(this, &SLiveLinkInputDeviceSourceFactory::OnSettingsAccepted)
+			.Text(LOCTEXT("AddSource", "Add"))
+		]
+	];
+#endif //WITH_EDITOR
+}
+
+FReply SLiveLinkInputDeviceSourceFactory::OnSettingsAccepted()
+{
+#if WITH_EDITOR
+	CastChecked<UScriptStruct>(StructOnScope->GetStruct())->CopyScriptStruct(&ConnectionSettings, StructOnScope->GetStructMemory());
+	OnConnectionSettingsAccepted.ExecuteIfBound(ConnectionSettings);
+#endif //WITH_EDITOR
+
+	return FReply::Handled();
+}
+
+#undef LOCTEXT_NAMESPACE
\ No newline at end of file
diff --git a/Engine/Plugins/VirtualProduction/LiveLinkInputDevice/Source/LiveLinkInputDevice/Private/SLiveLinkInputDeviceSourceFactory.h b/Engine/Plugins/VirtualProduction/LiveLinkInputDevice/Source/LiveLinkInputDevice/Private/SLiveLinkInputDeviceSourceFactory.h
new file mode 100644
--- /dev/null
+++ b/Engine/Plugins/VirtualProduction/LiveLinkInputDevice/Source/LiveLinkInputDevice/Private/SLiveLinkInputDeviceSourceFactory.h
@@ -0,0 +1,39 @@
+// Copyright Epic Games, Inc. All Rights Reserved.
+
+#pragma once
+
+#include "LiveLinkInputDeviceConnectionSettings.h"
+#include "Widgets/SCompoundWidget.h"
+#include "Widgets/DeclarativeSyntaxSupport.h"
+
+#if WITH_EDITOR
+#include "IStructureDetailsView.h"
+#endif //WITH_EDITOR
+
+#include "Input/Reply.h"
+
+struct FLiveLinkInputDeviceConnectionSettings;
+
+DECLARE_DELEGATE_OneParam(FOnLiveLinkInputDeviceConnectionSettingsAccepted, FLiveLinkInputDeviceConnectionSettings);
+
+class SLiveLinkInputDeviceSourceFactory : public SCompoundWidget
+{
+	SLATE_BEGIN_ARGS(SLiveLinkInputDeviceSourceFactory)
+	{}
+		SLATE_EVENT(FOnLiveLinkInputDeviceConnectionSettingsAccepted, OnConnectionSettingsAccepted)
+	SLATE_END_ARGS()
+
+	void Construct(const FArguments& Args);
+
+
+private:
+	FLiveLinkInputDeviceConnectionSettings ConnectionSettings;
+
+#if WITH_EDITOR
+	TSharedPtr<FStructOnScope> StructOnScope;
+	TSharedPtr<IStructureDetailsView> StructureDetailsView;
+#endif //WITH_EDITOR
+
+	FReply OnSettingsAccepted();
+	FOnLiveLinkInputDeviceConnectionSettingsAccepted OnConnectionSettingsAccepted;
+};
\ No newline at end of file
diff --git a/Engine/Plugins/VirtualProduction/LiveLinkInputDevice/Source/LiveLinkInputDevice/Public/LiveLinkInputDevice.h b/Engine/Plugins/VirtualProduction/LiveLinkInputDevice/Source/LiveLinkInputDevice/Public/LiveLinkInputDevice.h
new file mode 100644
--- /dev/null
+++ b/Engine/Plugins/VirtualProduction/LiveLinkInputDevice/Source/LiveLinkInputDevice/Public/LiveLinkInputDevice.h
@@ -0,0 +1,26 @@
+// Copyright Epic Games, Inc. All Rights Reserved.
+
+#pragma once
+
+#include "Modules/ModuleManager.h"
+
+DECLARE_LOG_CATEGORY_EXTERN(LogLiveLinkInputDevice, Log, All);
+
+/**
+ * This module provides access to a restricted set of the InputDevice system over Live Link without using the Game Input Device Subsystem.
+ *
+ * This allows us to poll / query the input devices from another thread for recording and animating purposes.
+ * */
+class FLiveLinkInputDeviceModule : public IModuleInterface
+{
+public:
+
+	static FLiveLinkInputDeviceModule& Get()
+	{
+		return FModuleManager::Get().LoadModuleChecked<FLiveLinkInputDeviceModule>(TEXT("LiveLinkInputDevice"));
+	}
+
+	/** IModuleInterface implementation */
+	virtual void StartupModule() override;
+	virtual void ShutdownModule() override;
+};
diff --git a/Engine/Plugins/VirtualProduction/LiveLinkInputDevice/Source/LiveLinkInputDevice/Public/LiveLinkInputDeviceConnectionSettings.h b/Engine/Plugins/VirtualProduction/LiveLinkInputDevice/Source/LiveLinkInputDevice/Public/LiveLinkInputDeviceConnectionSettings.h
new file mode 100644
--- /dev/null
+++ b/Engine/Plugins/VirtualProduction/LiveLinkInputDevice/Source/LiveLinkInputDevice/Public/LiveLinkInputDeviceConnectionSettings.h
@@ -0,0 +1,11 @@
+// Copyright Epic Games, Inc. All Rights Reserved.
+
+#pragma once
+
+#include "LiveLinkInputDeviceConnectionSettings.generated.h"
+
+USTRUCT()
+struct LIVELINKINPUTDEVICE_API FLiveLinkInputDeviceConnectionSettings
+{
+	GENERATED_BODY()
+};
diff --git a/Engine/Plugins/VirtualProduction/LiveLinkInputDevice/Source/LiveLinkInputDevice/Public/LiveLinkInputDeviceRole.h b/Engine/Plugins/VirtualProduction/LiveLinkInputDevice/Source/LiveLinkInputDevice/Public/LiveLinkInputDeviceRole.h
new file mode 100644
--- /dev/null
+++ b/Engine/Plugins/VirtualProduction/LiveLinkInputDevice/Source/LiveLinkInputDevice/Public/LiveLinkInputDeviceRole.h
@@ -0,0 +1,31 @@
+// Copyright Epic Games, Inc. All Rights Reserved.
+
+#pragma once
+
+#include "LiveLinkTypes.h"
+#include "Roles/LiveLinkBasicRole.h"
+
+#include "UObject/ObjectMacros.h"
+#include "UObject/UObjectGlobals.h"
+
+#include "LiveLinkInputDeviceRole.generated.h"
+
+/**
+ * Role associated with Input Device data.
+ */
+UCLASS(BlueprintType, meta = (DisplayName = "Input Device Role"))
+class LIVELINKINPUTDEVICE_API ULiveLinkInputDeviceRole : public ULiveLinkBasicRole
+{
+	GENERATED_BODY()
+
+public:
+	//~ Begin ULiveLinkRole interface
+	virtual UScriptStruct* GetStaticDataStruct() const override;
+	virtual UScriptStruct* GetFrameDataStruct() const override;
+	virtual UScriptStruct* GetBlueprintDataStruct() const override;
+
+	virtual bool InitializeBlueprintData(const FLiveLinkSubjectFrameData& InSourceData, FLiveLinkBlueprintDataStruct& OutBlueprintData) const override;
+
+	virtual FText GetDisplayName() const override;
+	//~ End ULiveLinkRole interface
+};
diff --git a/Engine/Plugins/VirtualProduction/LiveLinkInputDevice/Source/LiveLinkInputDevice/Public/LiveLinkInputDeviceSource.h b/Engine/Plugins/VirtualProduction/LiveLinkInputDevice/Source/LiveLinkInputDevice/Public/LiveLinkInputDeviceSource.h
new file mode 100644
--- /dev/null
+++ b/Engine/Plugins/VirtualProduction/LiveLinkInputDevice/Source/LiveLinkInputDevice/Public/LiveLinkInputDeviceSource.h
@@ -0,0 +1,114 @@
+// Copyright Epic Games, Inc. All Rights Reserved.
+
+#pragma once
+
+#include "ILiveLinkSource.h"
+#include "LiveLinkInputDeviceConnectionSettings.h"
+#include "LiveLinkInputDeviceMessageHandler.h"
+#include "LiveLinkInputDeviceSourceSettings.h"
+#include "LiveLinkInputDeviceTypes.h"
+#include "Misc/CoreMiscDefines.h"
+#include "Roles/LiveLinkTransformTypes.h"
+
+#include "Delegates/IDelegateInstance.h"
+#include "MessageEndpoint.h"
+#include "IMessageContext.h"
+#include "HAL/ThreadSafeBool.h"
+#include "HAL/Runnable.h"
+
+struct ULiveLinkInputDeviceSettings;
+
+class IInputDevice;
+class ILiveLinkClient;
+
+
+class LIVELINKINPUTDEVICE_API FLiveLinkInputDeviceSource : public ILiveLinkSource, public FRunnable, public TSharedFromThis<FLiveLinkInputDeviceSource>
+{
+public:
+
+	FLiveLinkInputDeviceSource(const FLiveLinkInputDeviceConnectionSettings& ConnectionSettings);
+
+	virtual ~FLiveLinkInputDeviceSource();
+
+	// Begin ILiveLinkSource Interface
+	virtual void ReceiveClient(ILiveLinkClient* InClient, FGuid InSourceGuid) override;
+	virtual void InitializeSettings(ULiveLinkSourceSettings* Settings) override;
+	virtual bool IsSourceStillValid() const override;
+
+	virtual bool RequestSourceShutdown() override;
+
+	virtual FText GetSourceType() const override { return SourceType; };
+	virtual FText GetSourceMachineName() const override { return SourceMachineName; }
+	virtual FText GetSourceStatus() const override { return SourceStatus; }
+
+	virtual TSubclassOf<ULiveLinkSourceSettings> GetSettingsClass() const override { return ULiveLinkInputDeviceSourceSettings::StaticClass(); }
+	// End ILiveLinkSource Interface
+
+	// Begin FRunnable Interface
+	virtual bool Init() override { return true; }
+	virtual uint32 Run() override;
+	void Start();
+	virtual void Stop() override;
+	virtual void Exit() override { }
+	// End FRunnable Interface
+
+private:
+	TOptional<FLiveLinkStaticDataStruct> PollGamepadStaticData(FInputDeviceId SubjectName);
+	TOptional<FLiveLinkFrameDataStruct>  PollGamepadFrameData(FInputDeviceId SubjectName);
+
+	void Send(FInputDeviceId DeviceId, FLiveLinkStaticDataStruct FrameDataToSend);
+	void Send(FInputDeviceId DeviceId, FLiveLinkFrameDataStruct StaticDataToSend);
+
+	/** Ticks each input device and calls SendControllerEvents to update the device with our message handler. */
+	bool TickInputDevices(const double InCurrentTime);
+
+	/** Time in seconds to wait before polling the device values again. */
+	double DevicePollWaitTime() const;
+
+private:
+	/** Update the our known gamepad state so that it can be reflected back to the LiveLink UI. */
+	void UpdateGamepadState(const bool bInState);
+
+	/** Create load all input devices by calling CreateInputDevice. */
+	void LoadInputPlugins();
+
+	ILiveLinkClient* Client;
+
+	// Our identifier in LiveLink
+	FGuid SourceGuid;
+
+	// Text objects that are reflected back in the Live Link UI.
+	FText SourceType;
+	FText SourceMachineName;
+	FText SourceStatus;
+	
+	// Threadsafe Bool for terminating the main thread loop
+	FThreadSafeBool Stopping;
+	
+	// Thread to run updates on
+	TUniquePtr<FRunnableThread> Thread;
+
+	// Name of the updates thread
+	FString ThreadName;
+
+	// List of subjects we've already encountered
+	TSet<FInputDeviceId> EncounteredSubjects;
+
+	// Deferred start delegate handle
+	FDelegateHandle DeferredStartDelegateHandle;
+
+	// Rate at which we poll the input devices.
+	uint32 LocalUpdateRateInHz = 100;
+
+	// Our local message handler.
+	TSharedPtr<class FLiveLinkInputDeviceMessageHandler> MessageHandler;
+
+	// List of known input devices registered with this source.
+	TArray<TSharedPtr<IInputDevice>> InputDevices;
+
+	// Bool to track if we have loaded our input plugins.  This should only be done once by querying the input device modular feature.
+	bool bHasLoadedInputPlugins = false;
+
+	// Indicates if we have ANY gamepads connected and thus are an active device for Live Link.
+	bool bIsGamepadConnected = false;
+};
diff --git a/Engine/Plugins/VirtualProduction/LiveLinkInputDevice/Source/LiveLinkInputDevice/Public/LiveLinkInputDeviceSourceFactory.h b/Engine/Plugins/VirtualProduction/LiveLinkInputDevice/Source/LiveLinkInputDevice/Public/LiveLinkInputDeviceSourceFactory.h
new file mode 100644
--- /dev/null
+++ b/Engine/Plugins/VirtualProduction/LiveLinkInputDevice/Source/LiveLinkInputDevice/Public/LiveLinkInputDeviceSourceFactory.h
@@ -0,0 +1,24 @@
+// Copyright Epic Games, Inc. All Rights Reserved.
+
+#pragma once
+
+#include "LiveLinkSourceFactory.h"
+#include "LiveLinkInputDeviceSource.h"
+#include "LiveLinkInputDeviceSourceFactory.generated.h"
+
+UCLASS()
+class LIVELINKINPUTDEVICE_API ULiveLinkInputDeviceSourceFactory : public ULiveLinkSourceFactory
+{
+public:
+	GENERATED_BODY()
+
+	virtual FText GetSourceDisplayName() const override;
+	virtual FText GetSourceTooltip() const override;
+
+	virtual EMenuType GetMenuType() const override { return EMenuType::SubPanel; }
+	virtual TSharedPtr<SWidget> BuildCreationPanel(FOnLiveLinkSourceCreated OnLiveLinkSourceCreated) const override;
+	virtual TSharedPtr<ILiveLinkSource> CreateSource(const FString& ConnectionString) const override;
+
+private:
+	void CreateSourceFromSettings(FLiveLinkInputDeviceConnectionSettings ConnectionSettings, FOnLiveLinkSourceCreated OnSourceCreated) const;
+};
\ No newline at end of file
diff --git a/Engine/Plugins/VirtualProduction/LiveLinkInputDevice/Source/LiveLinkInputDevice/Public/LiveLinkInputDeviceSourceSettings.h b/Engine/Plugins/VirtualProduction/LiveLinkInputDevice/Source/LiveLinkInputDevice/Public/LiveLinkInputDeviceSourceSettings.h
new file mode 100644
--- /dev/null
+++ b/Engine/Plugins/VirtualProduction/LiveLinkInputDevice/Source/LiveLinkInputDevice/Public/LiveLinkInputDeviceSourceSettings.h
@@ -0,0 +1,13 @@
+// Copyright Epic Games, Inc. All Rights Reserved.
+
+#pragma once
+
+#include "CoreMinimal.h"
+#include "LiveLinkSourceSettings.h"
+#include "LiveLinkInputDeviceSourceSettings.generated.h"
+
+UCLASS()
+class LIVELINKINPUTDEVICE_API ULiveLinkInputDeviceSourceSettings : public ULiveLinkSourceSettings
+{
+	GENERATED_BODY()
+};
diff --git a/Engine/Plugins/VirtualProduction/LiveLinkInputDevice/Source/LiveLinkInputDevice/Public/LiveLinkInputDeviceTypes.h b/Engine/Plugins/VirtualProduction/LiveLinkInputDevice/Source/LiveLinkInputDevice/Public/LiveLinkInputDeviceTypes.h
new file mode 100644
--- /dev/null
+++ b/Engine/Plugins/VirtualProduction/LiveLinkInputDevice/Source/LiveLinkInputDevice/Public/LiveLinkInputDeviceTypes.h
@@ -0,0 +1,142 @@
+// Copyright Epic Games, Inc. All Rights Reserved.
+
+#pragma once
+
+#include "Roles/LiveLinkBasicTypes.h"
+#include "Containers/Map.h"
+
+#include "LiveLinkInputDeviceTypes.generated.h"
+
+/**
+ * Struct for static Gamepad Input Device data
+ */
+USTRUCT(BlueprintType)
+struct LIVELINKINPUTDEVICE_API FLiveLinkGamepadInputDeviceStaticData : public FLiveLinkBaseStaticData
+{
+	GENERATED_BODY()
+};
+
+/**
+ * Struct for dynamic (per-frame) Gampead Input Device data
+ */
+USTRUCT(BlueprintType)
+struct LIVELINKINPUTDEVICE_API FLiveLinkGamepadInputDeviceFrameData : public FLiveLinkBaseFrameData
+{
+	GENERATED_BODY()
+
+	/** Given a key and its value apply it to the correct struct member variable. */
+	void ApplyValueFromKey(FName Key, float InValue);
+
+	UPROPERTY(EditAnywhere, BlueprintReadWrite, Category = "Gamepad Input Device")
+	float LeftAnalogX = 0;
+
+	UPROPERTY(EditAnywhere, BlueprintReadWrite, Category = "Gamepad Input Device")
+	float LeftAnalogY = 0;
+
+	UPROPERTY(EditAnywhere, BlueprintReadWrite, Category = "Gamepad Input Device")
+	float RightAnalogX = 0;
+
+	UPROPERTY(EditAnywhere, BlueprintReadWrite, Category = "Gamepad Input Device")
+	float RightAnalogY = 0;
+
+	UPROPERTY(EditAnywhere, BlueprintReadWrite, Category = "Gamepad Input Device")
+	float LeftTriggerAnalog = 0;
+
+	UPROPERTY(EditAnywhere, BlueprintReadWrite, Category = "Gamepad Input Device")
+	float RightTriggerAnalog = 0;
+
+	UPROPERTY(EditAnywhere, BlueprintReadWrite, Category = "Gamepad Input Device")
+	float LeftThumb = 0;
+
+	UPROPERTY(EditAnywhere, BlueprintReadWrite, Category = "Gamepad Input Device")
+	float RightThumb = 0;
+
+	UPROPERTY(EditAnywhere, BlueprintReadWrite, Category = "Gamepad Input Device")
+	float SpecialLeft = 0;
+
+	UPROPERTY(EditAnywhere, BlueprintReadWrite, Category = "Gamepad Input Device")
+	float SpecialLeft_X = 0;
+
+	UPROPERTY(EditAnywhere, BlueprintReadWrite, Category = "Gamepad Input Device")
+	float SpecialLeft_Y = 0;
+
+	UPROPERTY(EditAnywhere, BlueprintReadWrite, Category = "Gamepad Input Device")
+	float SpecialRight = 0;
+
+	UPROPERTY(EditAnywhere, BlueprintReadWrite, Category = "Gamepad Input Device")
+	float FaceButtonBottom = 0;
+
+	UPROPERTY(EditAnywhere, BlueprintReadWrite, Category = "Gamepad Input Device")
+	float FaceButtonRight = 0;
+
+	UPROPERTY(EditAnywhere, BlueprintReadWrite, Category = "Gamepad Input Device")
+	float FaceButtonLeft = 0;
+
+	UPROPERTY(EditAnywhere, BlueprintReadWrite, Category = "Gamepad Input Device")
+	float FaceButtonTop = 0;
+
+	UPROPERTY(EditAnywhere, BlueprintReadWrite, Category = "Gamepad Input Device")
+	float LeftShoulder = 0;
+
+	UPROPERTY(EditAnywhere, BlueprintReadWrite, Category = "Gamepad Input Device")
+	float RightShoulder = 0;
+
+	UPROPERTY(EditAnywhere, BlueprintReadWrite, Category = "Gamepad Input Device")
+	float LeftTriggerThreshold = 0;
+
+	UPROPERTY(EditAnywhere, BlueprintReadWrite, Category = "Gamepad Input Device")
+	float RightTriggerThreshold = 0;
+
+	UPROPERTY(EditAnywhere, BlueprintReadWrite, Category = "Gamepad Input Device")
+	float DPadUp = 0;
+
+	UPROPERTY(EditAnywhere, BlueprintReadWrite, Category = "Gamepad Input Device")
+	float DPadDown = 0;
+
+	UPROPERTY(EditAnywhere, BlueprintReadWrite, Category = "Gamepad Input Device")
+	float DPadRight = 0;
+
+	UPROPERTY(EditAnywhere, BlueprintReadWrite, Category = "Gamepad Input Device")
+	float DPadLeft = 0;
+
+	UPROPERTY(EditAnywhere, BlueprintReadWrite, Category = "Gamepad Input Device")
+	float LeftStickUp = 0;
+
+	UPROPERTY(EditAnywhere, BlueprintReadWrite, Category = "Gamepad Input Device")
+	float LeftStickDown = 0;
+
+	UPROPERTY(EditAnywhere, BlueprintReadWrite, Category = "Gamepad Input Device")
+	float LeftStickRight = 0;
+
+	UPROPERTY(EditAnywhere, BlueprintReadWrite, Category = "Gamepad Input Device")
+	float LeftStickLeft = 0;
+
+	UPROPERTY(EditAnywhere, BlueprintReadWrite, Category = "Gamepad Input Device")
+	float RightStickUp = 0;
+
+	UPROPERTY(EditAnywhere, BlueprintReadWrite, Category = "Gamepad Input Device")
+	float RightStickDown = 0;
+
+	UPROPERTY(EditAnywhere, BlueprintReadWrite, Category = "Gamepad Input Device")
+	float RightStickRight = 0;
+
+	UPROPERTY(EditAnywhere, BlueprintReadWrite, Category = "Gamepad Input Device")
+	float RightStickLeft = 0;
+};
+
+/**
+ * Facility structure to handle Preston MDR data in blueprint
+ */
+USTRUCT(BlueprintType)
+struct LIVELINKINPUTDEVICE_API FLiveLinkGamepadInputDeviceBlueprintData : public FLiveLinkBaseBlueprintData
+{
+	GENERATED_BODY()
+
+	/** Static data that should not change every frame */
+	UPROPERTY(EditAnywhere, BlueprintReadWrite, Category = "LiveLink")
+	FLiveLinkGamepadInputDeviceStaticData StaticData;
+
+	/** Dynamic data that can change every frame  */
+	UPROPERTY(EditAnywhere, BlueprintReadWrite, Category = "LiveLink")
+	FLiveLinkGamepadInputDeviceFrameData FrameData;
+};
diff --git a/Engine/Programs/LiveLinkHub/Config/DefaultEngine.ini b/Engine/Programs/LiveLinkHub/Config/DefaultEngine.ini
--- a/Engine/Programs/LiveLinkHub/Config/DefaultEngine.ini
+++ b/Engine/Programs/LiveLinkHub/Config/DefaultEngine.ini
@@ -7,13 +7,13 @@ GameNameShort=LLH
 +ProgramEnabledPlugins="ContentBrowserAssetDataSource"
 +ProgramEnabledPlugins="PropertyAccessEditor"
 +ProgramEnabledPlugins="PythonScriptPlugin"
++ProgramEnabledPlugins="LiveLinkInputDevice"
 
 +ProgramOptionalPlugins="AppleARKitFaceSupport"
 #+ProgramOptionalPlugins="LiveLinkViconDataStream"
 +ProgramOptionalPlugins="MocopiLiveLink"
 +ProgramOptionalPlugins="OptitrackLiveLink"
 
-
 [/Script/MacTargetPlatform.XcodeProjectSettings]
 TemplateMacPlist=(FilePath="/Engine/Source/Programs/LiveLinkHub/Resources/Mac/Info.plist")
 
diff --git a/Engine/Source/Programs/LiveLinkHub/LiveLinkHub.Target.cs b/Engine/Source/Programs/LiveLinkHub/LiveLinkHub.Target.cs
--- a/Engine/Source/Programs/LiveLinkHub/LiveLinkHub.Target.cs
+++ b/Engine/Source/Programs/LiveLinkHub/LiveLinkHub.Target.cs
@@ -26,7 +26,7 @@ public class LiveLinkHubTarget : TargetRules
 		AdditionalPlugins.Add("QuicMessaging");
 		AdditionalPlugins.Add("PropertyAccessEditor");
 		AdditionalPlugins.Add("PythonScriptPlugin");
-
+		AdditionalPlugins.Add("LiveLinkInputDevice");
 		OptionalPlugins.AddRange(new string[]
 		{
 			"AppleARKitFaceSupport",
