Attempt No. 2 to fix 416438 on Mac.

diff --git a/Engine/Plugins/VirtualProduction/LiveLinkInputDevice/LiveLinkInputDevice.uplugin b/Engine/Plugins/VirtualProduction/LiveLinkInputDevice/LiveLinkInputDevice.uplugin
--- a/Engine/Plugins/VirtualProduction/LiveLinkInputDevice/LiveLinkInputDevice.uplugin
+++ b/Engine/Plugins/VirtualProduction/LiveLinkInputDevice/LiveLinkInputDevice.uplugin
@@ -16,7 +16,6 @@
 	"Installed": false,
 	"EnabledByDefault": false,
 	"SupportedPrograms": ["LiveLinkHub"],
-	"SupportedTargetPlatforms": ["Win64"],
 	"Modules": [
 		{
 			"Name": "LiveLinkInputDevice",
