Move input device type to accommodate accessing input device data from the Live Link Control Rig plugin.

diff --git a/Engine/Plugins/Experimental/LiveLinkControlRig/Source/LiveLinkControlRig/Private/LiveLinkInputDeviceRigUnits.cpp b/Engine/Plugins/Experimental/LiveLinkControlRig/Source/LiveLinkControlRig/Private/LiveLinkInputDeviceRigUnits.cpp
new file mode 100644
--- /dev/null
+++ b/Engine/Plugins/Experimental/LiveLinkControlRig/Source/LiveLinkControlRig/Private/LiveLinkInputDeviceRigUnits.cpp
@@ -0,0 +1,45 @@
+// Copyright Epic Games, Inc. All Rights Reserved.
+
+#include "LiveLinkInputDeviceRigUnits.h"
+
+#include "Features/IModularFeatures.h"
+#include "Roles/LiveLinkInputDeviceRole.h"
+
+#include "ILiveLinkClient.h"
+#include "Roles/LiveLinkInputDeviceTypes.h"
+
+namespace LiveLinkInputDeviceRigUtils
+{
+
+ILiveLinkClient* TryGetLiveLinkClient()
+{
+	static ILiveLinkClient* LiveLinkClient = []() -> ILiveLinkClient*
+	{
+		IModularFeatures& ModularFeatures = IModularFeatures::Get();
+		if (ModularFeatures.IsModularFeatureAvailable(ILiveLinkClient::ModularFeatureName))
+		{
+			return static_cast<ILiveLinkClient*>(&ModularFeatures.GetModularFeature<ILiveLinkClient>(ILiveLinkClient::ModularFeatureName));
+		}
+		return nullptr;
+	}();
+	return LiveLinkClient;
+}
+
+}
+
+FRigUnit_LiveLinkEvaluateInputDeviceValue_Execute()
+{
+	// Get value by property name from input device gamepad.
+	if (ILiveLinkClient* LiveLinkClient = LiveLinkInputDeviceRigUtils::TryGetLiveLinkClient())
+	{
+		FLiveLinkSubjectFrameData FrameData;
+		const bool bHaveFrameData =	LiveLinkClient->EvaluateFrame_AnyThread(SubjectName, ULiveLinkInputDeviceRole::StaticClass(),FrameData);
+		if (bHaveFrameData)
+		{
+			if (FLiveLinkGamepadInputDeviceFrameData* Data = FrameData.FrameData.Cast<FLiveLinkGamepadInputDeviceFrameData>())
+			{
+				InputDeviceData = *Data;
+			}
+		}
+	}
+}
diff --git a/Engine/Plugins/Experimental/LiveLinkControlRig/Source/LiveLinkControlRig/Public/LiveLinkInputDeviceRigUnits.h b/Engine/Plugins/Experimental/LiveLinkControlRig/Source/LiveLinkControlRig/Public/LiveLinkInputDeviceRigUnits.h
new file mode 100644
--- /dev/null
+++ b/Engine/Plugins/Experimental/LiveLinkControlRig/Source/LiveLinkControlRig/Public/LiveLinkInputDeviceRigUnits.h
@@ -0,0 +1,28 @@
+// Copyright Epic Games, Inc. All Rights Reserved.
+
+#pragma once
+
+#include "ControlRigDefines.h"
+#include "LiveLinkTypes.h"
+#include "Roles/LiveLinkInputDeviceTypes.h"
+#include "Units/RigUnit.h"
+
+#include "LiveLinkInputDeviceRigUnits.generated.h"
+
+USTRUCT(meta = (DisplayName = "Get Live Link Input Device Data", Category = "Live Link"))
+struct LIVELINKCONTROLRIG_API FRigUnit_LiveLinkEvaluateInputDeviceValue : public FRigUnit
+{
+	GENERATED_BODY()
+
+	/** Execute logic for this rig unit */
+	RIGVM_METHOD()
+	virtual void Execute() override;
+
+	/** Name of the subject we would like to get Live Link device data. This should be the subject attached to the gamepad. */
+	UPROPERTY(meta = (Input))
+	FName SubjectName;
+
+	/** Gampad device data for the current frame. */
+	UPROPERTY(meta = (Output))
+	FLiveLinkGamepadInputDeviceFrameData InputDeviceData;
+};
diff --git a/Engine/Plugins/VirtualProduction/LiveLinkInputDevice/Source/LiveLinkInputDevice/Private/LiveLinkInputDeviceMessageHandler.cpp b/Engine/Plugins/VirtualProduction/LiveLinkInputDevice/Source/LiveLinkInputDevice/Private/LiveLinkInputDeviceMessageHandler.cpp
--- a/Engine/Plugins/VirtualProduction/LiveLinkInputDevice/Source/LiveLinkInputDevice/Private/LiveLinkInputDeviceMessageHandler.cpp
+++ b/Engine/Plugins/VirtualProduction/LiveLinkInputDevice/Source/LiveLinkInputDevice/Private/LiveLinkInputDeviceMessageHandler.cpp
@@ -1,28 +1,178 @@
 // Copyright Epic Games, Inc. All Rights Reserved.
 
 #include "LiveLinkInputDeviceMessageHandler.h"
-#include "LiveLinkInputDeviceTypes.h"
+#include "Roles/LiveLinkInputDeviceTypes.h"
 #include "Misc/CoreMiscDefines.h"
 
+#include "GenericPlatform/GenericApplicationMessageHandler.h"
+
+namespace UE::InputDevice::Private
+{
+
+void SetValueWithKey(FName Key, float InValue, FLiveLinkGamepadInputDeviceFrameData& OutData)
+{
+	if (Key == FGamepadKeyNames::LeftAnalogX)
+	{
+		OutData.LeftAnalogX = InValue;
+	}
+	else
+	if (Key == FGamepadKeyNames::LeftAnalogY)
+	{
+		OutData.LeftAnalogY = InValue;
+	}
+	else
+	if (Key == FGamepadKeyNames::RightAnalogX)
+	{
+		OutData.RightAnalogX = InValue;
+	}
+	else
+	if (Key == FGamepadKeyNames::RightAnalogY)
+	{
+		OutData.RightAnalogY = InValue;
+	}
+	else
+	if (Key == FGamepadKeyNames::LeftTriggerAnalog)
+	{
+		OutData.LeftTriggerAnalog = InValue;
+	}
+	else
+	if (Key == FGamepadKeyNames::RightTriggerAnalog)
+	{
+		OutData.RightTriggerAnalog = InValue;
+	}
+	else
+	if (Key == FGamepadKeyNames::LeftThumb)
+	{
+		OutData.LeftThumb = InValue;
+	}
+	else
+	if (Key == FGamepadKeyNames::RightThumb)
+	{
+		OutData.RightThumb = InValue;
+	}
+	else
+	if (Key == FGamepadKeyNames::SpecialLeft)
+	{
+		OutData.SpecialLeft = InValue;
+	}
+	else
+	if (Key == FGamepadKeyNames::SpecialLeft_X)
+	{
+		OutData.SpecialLeft_X = InValue;
+	}
+	else if (Key == FGamepadKeyNames::SpecialLeft_Y)
+	{
+		OutData.SpecialLeft_Y = InValue;
+	}
+	else if (Key == FGamepadKeyNames::SpecialRight)
+	{
+		OutData.SpecialRight = InValue;
+	}
+	else if (Key == FGamepadKeyNames::FaceButtonBottom)
+	{
+		OutData.FaceButtonBottom = InValue;
+	}
+	else if (Key == FGamepadKeyNames::FaceButtonRight)
+	{
+		OutData.FaceButtonRight = InValue;
+	}
+	else if (Key == FGamepadKeyNames::FaceButtonLeft)
+	{
+		OutData.FaceButtonLeft = InValue;
+	}
+	else if (Key == FGamepadKeyNames::FaceButtonTop)
+	{
+		OutData.FaceButtonTop = InValue;
+	}
+	else if (Key == FGamepadKeyNames::LeftShoulder)
+	{
+		OutData.LeftShoulder = InValue;
+	}
+	else if (Key == FGamepadKeyNames::RightShoulder)
+	{
+		OutData.RightShoulder = InValue;
+	}
+	else if (Key == FGamepadKeyNames::LeftTriggerThreshold)
+	{
+		OutData.LeftTriggerThreshold = InValue;
+	}
+	else if (Key == FGamepadKeyNames::RightTriggerThreshold)
+	{
+		OutData.RightTriggerThreshold = InValue;
+	}
+	else if (Key == FGamepadKeyNames::DPadUp)
+	{
+		OutData.DPadUp = InValue;
+	}
+	else if (Key == FGamepadKeyNames::DPadDown)
+	{
+		OutData.DPadDown = InValue;
+	}
+	else if (Key == FGamepadKeyNames::DPadRight)
+	{
+		OutData.DPadRight = InValue;
+	}
+	else if (Key == FGamepadKeyNames::DPadLeft)
+	{
+		OutData.DPadLeft = InValue;
+	}
+	else if (Key == FGamepadKeyNames::LeftStickUp)
+	{
+		OutData.LeftStickUp = InValue;
+	}
+	else if (Key == FGamepadKeyNames::LeftStickDown)
+	{
+		OutData.LeftStickDown = InValue;
+	}
+	else if (Key == FGamepadKeyNames::LeftStickRight)
+	{
+		OutData.LeftStickRight = InValue;
+	}
+	else if (Key == FGamepadKeyNames::LeftStickLeft)
+	{
+		OutData.LeftStickLeft = InValue;
+	}
+	else if (Key == FGamepadKeyNames::RightStickUp)
+	{
+		OutData.RightStickUp = InValue;
+	}
+	else if (Key == FGamepadKeyNames::RightStickDown)
+	{
+		OutData.RightStickDown = InValue;
+	}
+	else if (Key == FGamepadKeyNames::RightStickRight)
+	{
+		OutData.RightStickRight = InValue;
+	}
+	else if (Key == FGamepadKeyNames::RightStickLeft)
+	{
+		OutData.RightStickLeft = InValue;
+	}
+	else
+	{
+		ensureMsgf(false, TEXT("Unknown key value %s"), *Key.ToString());
+	}
+}
+}
 
 bool FLiveLinkInputDeviceMessageHandler::OnControllerAnalog(FGamepadKeyNames::Type KeyName, FPlatformUserId PlatformUserId, FInputDeviceId InputDeviceId, float AnalogValue)
 {
 	FLiveLinkGamepadInputDeviceFrameData& InputDeviceFrameData = CurrentFrameDataValues.FindOrAdd(InputDeviceId);
-	InputDeviceFrameData.ApplyValueFromKey(KeyName, AnalogValue);
+	UE::InputDevice::Private::SetValueWithKey(KeyName, AnalogValue, InputDeviceFrameData);
 	return true;
 }
 
 bool FLiveLinkInputDeviceMessageHandler::OnControllerButtonPressed(FGamepadKeyNames::Type KeyName, FPlatformUserId PlatformUserId, FInputDeviceId InputDeviceId, bool IsRepeat)
 {
 	FLiveLinkGamepadInputDeviceFrameData& InputDeviceFrameData = CurrentFrameDataValues.FindOrAdd(InputDeviceId);
-	InputDeviceFrameData.ApplyValueFromKey(KeyName, 1);
+	UE::InputDevice::Private::SetValueWithKey(KeyName, 1, InputDeviceFrameData);
 	return true;
 }
 
 bool FLiveLinkInputDeviceMessageHandler::OnControllerButtonReleased(FGamepadKeyNames::Type KeyName, FPlatformUserId PlatformUserId, FInputDeviceId InputDeviceId, bool IsRepeat)
 {
 	FLiveLinkGamepadInputDeviceFrameData& InputDeviceFrameData = CurrentFrameDataValues.FindOrAdd(InputDeviceId);
-	InputDeviceFrameData.ApplyValueFromKey(KeyName, 0);
+	UE::InputDevice::Private::SetValueWithKey(KeyName, 0, InputDeviceFrameData);
 	return true;
 }
 
diff --git a/Engine/Plugins/VirtualProduction/LiveLinkInputDevice/Source/LiveLinkInputDevice/Private/LiveLinkInputDeviceMessageHandler.h b/Engine/Plugins/VirtualProduction/LiveLinkInputDevice/Source/LiveLinkInputDevice/Private/LiveLinkInputDeviceMessageHandler.h
--- a/Engine/Plugins/VirtualProduction/LiveLinkInputDevice/Source/LiveLinkInputDevice/Private/LiveLinkInputDeviceMessageHandler.h
+++ b/Engine/Plugins/VirtualProduction/LiveLinkInputDevice/Source/LiveLinkInputDevice/Private/LiveLinkInputDeviceMessageHandler.h
@@ -3,7 +3,7 @@
 #pragma once
 
 #include "GenericPlatform/GenericApplicationMessageHandler.h"
-#include "LiveLinkInputDeviceTypes.h"
+#include "Roles/LiveLinkInputDeviceTypes.h"
 #include "Containers/Map.h"
 #include "Misc/CoreMiscDefines.h"
 
diff --git a/Engine/Plugins/VirtualProduction/LiveLinkInputDevice/Source/LiveLinkInputDevice/Private/LiveLinkInputDeviceRole.cpp b/Engine/Plugins/VirtualProduction/LiveLinkInputDevice/Source/LiveLinkInputDevice/Private/LiveLinkInputDeviceRole.cpp
deleted file mode 100644
diff --git a/Engine/Plugins/VirtualProduction/LiveLinkInputDevice/Source/LiveLinkInputDevice/Private/LiveLinkInputDeviceSource.cpp b/Engine/Plugins/VirtualProduction/LiveLinkInputDevice/Source/LiveLinkInputDevice/Private/LiveLinkInputDeviceSource.cpp
--- a/Engine/Plugins/VirtualProduction/LiveLinkInputDevice/Source/LiveLinkInputDevice/Private/LiveLinkInputDeviceSource.cpp
+++ b/Engine/Plugins/VirtualProduction/LiveLinkInputDevice/Source/LiveLinkInputDevice/Private/LiveLinkInputDeviceSource.cpp
@@ -12,12 +12,12 @@
 #include "InputCoreTypes.h"
 #include "Misc/CoreDelegates.h"
 #include "Roles/LiveLinkTransformRole.h"
+#include "Roles/LiveLinkInputDeviceTypes.h"
+#include "Roles/LiveLinkInputDeviceRole.h"
 
 #include "ILiveLinkClient.h"
 #include "LiveLinkInputDevice.h"
-#include "LiveLinkInputDeviceRole.h"
 #include "LiveLinkInputDeviceSourceSettings.h"
-#include "LiveLinkInputDeviceTypes.h"
 #include "LiveLinkTypes.h"
 
 #include "LiveLinkInputDeviceMessageHandler.h"
diff --git a/Engine/Plugins/VirtualProduction/LiveLinkInputDevice/Source/LiveLinkInputDevice/Private/LiveLinkInputDeviceTypes.cpp b/Engine/Plugins/VirtualProduction/LiveLinkInputDevice/Source/LiveLinkInputDevice/Private/LiveLinkInputDeviceTypes.cpp
deleted file mode 100644
diff --git a/Engine/Plugins/VirtualProduction/LiveLinkInputDevice/Source/LiveLinkInputDevice/Public/LiveLinkInputDeviceRole.h b/Engine/Plugins/VirtualProduction/LiveLinkInputDevice/Source/LiveLinkInputDevice/Public/LiveLinkInputDeviceRole.h
deleted file mode 100644
diff --git a/Engine/Plugins/VirtualProduction/LiveLinkInputDevice/Source/LiveLinkInputDevice/Public/LiveLinkInputDeviceSource.h b/Engine/Plugins/VirtualProduction/LiveLinkInputDevice/Source/LiveLinkInputDevice/Public/LiveLinkInputDeviceSource.h
--- a/Engine/Plugins/VirtualProduction/LiveLinkInputDevice/Source/LiveLinkInputDevice/Public/LiveLinkInputDeviceSource.h
+++ b/Engine/Plugins/VirtualProduction/LiveLinkInputDevice/Source/LiveLinkInputDevice/Public/LiveLinkInputDeviceSource.h
@@ -6,8 +6,9 @@
 #include "LiveLinkInputDeviceConnectionSettings.h"
 #include "LiveLinkInputDeviceMessageHandler.h"
 #include "LiveLinkInputDeviceSourceSettings.h"
-#include "LiveLinkInputDeviceTypes.h"
 #include "Misc/CoreMiscDefines.h"
+
+#include "Roles/LiveLinkInputDeviceTypes.h"
 #include "Roles/LiveLinkTransformTypes.h"
 
 #include "Delegates/IDelegateInstance.h"
diff --git a/Engine/Plugins/VirtualProduction/LiveLinkInputDevice/Source/LiveLinkInputDevice/Public/LiveLinkInputDeviceTypes.h b/Engine/Plugins/VirtualProduction/LiveLinkInputDevice/Source/LiveLinkInputDevice/Public/LiveLinkInputDeviceTypes.h
deleted file mode 100644
diff --git a/Engine/Source/Runtime/LiveLinkInterface/Private/Roles/LiveLinkInputDeviceRole.cpp b/Engine/Source/Runtime/LiveLinkInterface/Private/Roles/LiveLinkInputDeviceRole.cpp
new file mode 100644
--- /dev/null
+++ b/Engine/Source/Runtime/LiveLinkInterface/Private/Roles/LiveLinkInputDeviceRole.cpp
@@ -0,0 +1,46 @@
+// Copyright Epic Games, Inc. All Rights Reserved.
+
+
+#include "Roles/LiveLinkInputDeviceRole.h"
+#include "Roles/LiveLinkInputDeviceTypes.h"
+
+#define LOCTEXT_NAMESPACE "LiveLinkInputDeviceRole"
+
+UScriptStruct* ULiveLinkInputDeviceRole::GetStaticDataStruct() const
+{
+	return FLiveLinkGamepadInputDeviceStaticData::StaticStruct();
+}
+
+UScriptStruct* ULiveLinkInputDeviceRole::GetFrameDataStruct() const
+{
+	return FLiveLinkGamepadInputDeviceFrameData::StaticStruct();
+}
+
+UScriptStruct* ULiveLinkInputDeviceRole::GetBlueprintDataStruct() const
+{
+	return FLiveLinkGamepadInputDeviceBlueprintData::StaticStruct();
+}
+
+bool ULiveLinkInputDeviceRole::InitializeBlueprintData(const FLiveLinkSubjectFrameData& InSourceData, FLiveLinkBlueprintDataStruct& OutBlueprintData) const
+{
+	bool bSuccess = false;
+
+	FLiveLinkGamepadInputDeviceBlueprintData* BlueprintData = OutBlueprintData.Cast<FLiveLinkGamepadInputDeviceBlueprintData>();
+	const FLiveLinkGamepadInputDeviceStaticData* StaticData = InSourceData.StaticData.Cast<FLiveLinkGamepadInputDeviceStaticData>();
+	const FLiveLinkGamepadInputDeviceFrameData* FrameData = InSourceData.FrameData.Cast<FLiveLinkGamepadInputDeviceFrameData>();
+	if (BlueprintData && StaticData && FrameData)
+	{
+		GetStaticDataStruct()->CopyScriptStruct(&BlueprintData->StaticData, StaticData);
+		GetFrameDataStruct()->CopyScriptStruct(&BlueprintData->FrameData, FrameData);
+		bSuccess = true;
+	}
+
+	return bSuccess;
+}
+
+FText ULiveLinkInputDeviceRole::GetDisplayName() const
+{
+	return LOCTEXT("LiveLinkInputDevice", "Live Link Input Device");
+}
+
+#undef LOCTEXT_NAMESPACE
diff --git a/Engine/Source/Runtime/LiveLinkInterface/Private/Roles/LiveLinkInputDeviceTypes.cpp b/Engine/Source/Runtime/LiveLinkInterface/Private/Roles/LiveLinkInputDeviceTypes.cpp
new file mode 100644
--- /dev/null
+++ b/Engine/Source/Runtime/LiveLinkInterface/Private/Roles/LiveLinkInputDeviceTypes.cpp
@@ -0,0 +1,3 @@
+// Copyright Epic Games, Inc. All Rights Reserved.
+
+#include "Roles/LiveLinkInputDeviceTypes.h"
diff --git a/Engine/Source/Runtime/LiveLinkInterface/Public/Roles/LiveLinkInputDeviceRole.h b/Engine/Source/Runtime/LiveLinkInterface/Public/Roles/LiveLinkInputDeviceRole.h
new file mode 100644
--- /dev/null
+++ b/Engine/Source/Runtime/LiveLinkInterface/Public/Roles/LiveLinkInputDeviceRole.h
@@ -0,0 +1,31 @@
+// Copyright Epic Games, Inc. All Rights Reserved.
+
+#pragma once
+
+#include "LiveLinkTypes.h"
+#include "Roles/LiveLinkBasicRole.h"
+
+#include "UObject/ObjectMacros.h"
+#include "UObject/UObjectGlobals.h"
+
+#include "LiveLinkInputDeviceRole.generated.h"
+
+/**
+ * Role associated with Input Device data.
+ */
+UCLASS(BlueprintType, meta = (DisplayName = "Input Device Role"))
+class LIVELINKINTERFACE_API ULiveLinkInputDeviceRole : public ULiveLinkBasicRole
+{
+	GENERATED_BODY()
+
+public:
+	//~ Begin ULiveLinkRole interface
+	virtual UScriptStruct* GetStaticDataStruct() const override;
+	virtual UScriptStruct* GetFrameDataStruct() const override;
+	virtual UScriptStruct* GetBlueprintDataStruct() const override;
+
+	virtual bool InitializeBlueprintData(const FLiveLinkSubjectFrameData& InSourceData, FLiveLinkBlueprintDataStruct& OutBlueprintData) const override;
+
+	virtual FText GetDisplayName() const override;
+	//~ End ULiveLinkRole interface
+};
diff --git a/Engine/Source/Runtime/LiveLinkInterface/Public/Roles/LiveLinkInputDeviceTypes.h b/Engine/Source/Runtime/LiveLinkInterface/Public/Roles/LiveLinkInputDeviceTypes.h
new file mode 100644
--- /dev/null
+++ b/Engine/Source/Runtime/LiveLinkInterface/Public/Roles/LiveLinkInputDeviceTypes.h
@@ -0,0 +1,139 @@
+// Copyright Epic Games, Inc. All Rights Reserved.
+
+#pragma once
+
+#include "Roles/LiveLinkBasicTypes.h"
+#include "Containers/Map.h"
+
+#include "LiveLinkInputDeviceTypes.generated.h"
+
+/**
+ * Struct for static Gamepad Input Device data
+ */
+USTRUCT(BlueprintType)
+struct LIVELINKINTERFACE_API FLiveLinkGamepadInputDeviceStaticData : public FLiveLinkBaseStaticData
+{
+	GENERATED_BODY()
+};
+
+/**
+ * Struct for dynamic (per-frame) Gampead Input Device data
+ */
+USTRUCT(BlueprintType)
+struct LIVELINKINTERFACE_API FLiveLinkGamepadInputDeviceFrameData : public FLiveLinkBaseFrameData
+{
+	GENERATED_BODY()
+
+	UPROPERTY(EditAnywhere, BlueprintReadWrite, Category = "Gamepad Input Device")
+	float LeftAnalogX = 0;
+
+	UPROPERTY(EditAnywhere, BlueprintReadWrite, Category = "Gamepad Input Device")
+	float LeftAnalogY = 0;
+
+	UPROPERTY(EditAnywhere, BlueprintReadWrite, Category = "Gamepad Input Device")
+	float RightAnalogX = 0;
+
+	UPROPERTY(EditAnywhere, BlueprintReadWrite, Category = "Gamepad Input Device")
+	float RightAnalogY = 0;
+
+	UPROPERTY(EditAnywhere, BlueprintReadWrite, Category = "Gamepad Input Device")
+	float LeftTriggerAnalog = 0;
+
+	UPROPERTY(EditAnywhere, BlueprintReadWrite, Category = "Gamepad Input Device")
+	float RightTriggerAnalog = 0;
+
+	UPROPERTY(EditAnywhere, BlueprintReadWrite, Category = "Gamepad Input Device")
+	float LeftThumb = 0;
+
+	UPROPERTY(EditAnywhere, BlueprintReadWrite, Category = "Gamepad Input Device")
+	float RightThumb = 0;
+
+	UPROPERTY(EditAnywhere, BlueprintReadWrite, Category = "Gamepad Input Device")
+	float SpecialLeft = 0;
+
+	UPROPERTY(EditAnywhere, BlueprintReadWrite, Category = "Gamepad Input Device")
+	float SpecialLeft_X = 0;
+
+	UPROPERTY(EditAnywhere, BlueprintReadWrite, Category = "Gamepad Input Device")
+	float SpecialLeft_Y = 0;
+
+	UPROPERTY(EditAnywhere, BlueprintReadWrite, Category = "Gamepad Input Device")
+	float SpecialRight = 0;
+
+	UPROPERTY(EditAnywhere, BlueprintReadWrite, Category = "Gamepad Input Device")
+	float FaceButtonBottom = 0;
+
+	UPROPERTY(EditAnywhere, BlueprintReadWrite, Category = "Gamepad Input Device")
+	float FaceButtonRight = 0;
+
+	UPROPERTY(EditAnywhere, BlueprintReadWrite, Category = "Gamepad Input Device")
+	float FaceButtonLeft = 0;
+
+	UPROPERTY(EditAnywhere, BlueprintReadWrite, Category = "Gamepad Input Device")
+	float FaceButtonTop = 0;
+
+	UPROPERTY(EditAnywhere, BlueprintReadWrite, Category = "Gamepad Input Device")
+	float LeftShoulder = 0;
+
+	UPROPERTY(EditAnywhere, BlueprintReadWrite, Category = "Gamepad Input Device")
+	float RightShoulder = 0;
+
+	UPROPERTY(EditAnywhere, BlueprintReadWrite, Category = "Gamepad Input Device")
+	float LeftTriggerThreshold = 0;
+
+	UPROPERTY(EditAnywhere, BlueprintReadWrite, Category = "Gamepad Input Device")
+	float RightTriggerThreshold = 0;
+
+	UPROPERTY(EditAnywhere, BlueprintReadWrite, Category = "Gamepad Input Device")
+	float DPadUp = 0;
+
+	UPROPERTY(EditAnywhere, BlueprintReadWrite, Category = "Gamepad Input Device")
+	float DPadDown = 0;
+
+	UPROPERTY(EditAnywhere, BlueprintReadWrite, Category = "Gamepad Input Device")
+	float DPadRight = 0;
+
+	UPROPERTY(EditAnywhere, BlueprintReadWrite, Category = "Gamepad Input Device")
+	float DPadLeft = 0;
+
+	UPROPERTY(EditAnywhere, BlueprintReadWrite, Category = "Gamepad Input Device")
+	float LeftStickUp = 0;
+
+	UPROPERTY(EditAnywhere, BlueprintReadWrite, Category = "Gamepad Input Device")
+	float LeftStickDown = 0;
+
+	UPROPERTY(EditAnywhere, BlueprintReadWrite, Category = "Gamepad Input Device")
+	float LeftStickRight = 0;
+
+	UPROPERTY(EditAnywhere, BlueprintReadWrite, Category = "Gamepad Input Device")
+	float LeftStickLeft = 0;
+
+	UPROPERTY(EditAnywhere, BlueprintReadWrite, Category = "Gamepad Input Device")
+	float RightStickUp = 0;
+
+	UPROPERTY(EditAnywhere, BlueprintReadWrite, Category = "Gamepad Input Device")
+	float RightStickDown = 0;
+
+	UPROPERTY(EditAnywhere, BlueprintReadWrite, Category = "Gamepad Input Device")
+	float RightStickRight = 0;
+
+	UPROPERTY(EditAnywhere, BlueprintReadWrite, Category = "Gamepad Input Device")
+	float RightStickLeft = 0;
+};
+
+/**
+ * Facility structure to handle Preston MDR data in blueprint
+ */
+USTRUCT(BlueprintType)
+struct LIVELINKINTERFACE_API FLiveLinkGamepadInputDeviceBlueprintData : public FLiveLinkBaseBlueprintData
+{
+	GENERATED_BODY()
+
+	/** Static data that should not change every frame */
+	UPROPERTY(EditAnywhere, BlueprintReadWrite, Category = "LiveLink")
+	FLiveLinkGamepadInputDeviceStaticData StaticData;
+
+	/** Dynamic data that can change every frame  */
+	UPROPERTY(EditAnywhere, BlueprintReadWrite, Category = "LiveLink")
+	FLiveLinkGamepadInputDeviceFrameData FrameData;
+};
