Use IPlatformInputDeviceMapper::Get().GetOnInputDeviceConnectionChange() to detect game pad changes when running in Live Link Hub.

diff --git a/Engine/Plugins/Animation/LiveLinkHub/Source/LiveLinkHub/Private/LiveLinkHubApplication.h b/Engine/Plugins/Animation/LiveLinkHub/Source/LiveLinkHub/Private/LiveLinkHubApplication.h
--- a/Engine/Plugins/Animation/LiveLinkHub/Source/LiveLinkHub/Private/LiveLinkHubApplication.h
+++ b/Engine/Plugins/Animation/LiveLinkHub/Source/LiveLinkHub/Private/LiveLinkHubApplication.h
@@ -4,11 +4,13 @@
 
 #include "Commandlets/Commandlet.h"
 #include "CoreMinimal.h"
-#include "Engine/Engine.h"
 #include "Editor/EditorPerformanceSettings.h"
-#include "Misc/CoreDelegates.h"
+#include "Engine/Engine.h"
+#include "Framework/Application/SlateApplication.h"
 #include "LiveLinkHub.h"
+#include "LiveLinkHubInputProcessor.h"
 #include "Misc/App.h"
+#include "Misc/CoreDelegates.h"
 #include "Runtime/Launch/Resources/Version.h"
 #include "UObject/UObjectGlobals.h"
 
@@ -23,6 +25,8 @@ void LiveLinkHubLoop(const TSharedPtr<FLiveLinkHub>& LiveLinkHub)
 	// Disable throttling for the hub
 	GetMutableDefault<UEditorPerformanceSettings>()->bThrottleCPUWhenNotForeground = false;
 
+	check(FSlateApplication::IsInitialized());
+	FSlateApplication::Get().RegisterInputPreProcessor(MakeShared<FLiveLinkHubInputProcessor>());
 	{
 		UE_LOG(LogLiveLinkHubApplication, Display, TEXT("LiveLinkHub Initialized (Version: %d.%d)"), ENGINE_MAJOR_VERSION, ENGINE_MINOR_VERSION);
 
@@ -41,6 +45,8 @@ void LiveLinkHubLoop(const TSharedPtr<FLiveLinkHub>& LiveLinkHub)
 
 			CommandletHelpers::TickEngine(nullptr, DeltaTime);
 
+			FSlateApplication::Get().PollGameDeviceState();
+
 			// This is normally ticked by OnSamplingInput.
 			LiveLinkHub->Tick();
 
diff --git a/Engine/Plugins/Animation/LiveLinkHub/Source/LiveLinkHub/Private/LiveLinkHubInputProcessor.cpp b/Engine/Plugins/Animation/LiveLinkHub/Source/LiveLinkHub/Private/LiveLinkHubInputProcessor.cpp
new file mode 100644
--- /dev/null
+++ b/Engine/Plugins/Animation/LiveLinkHub/Source/LiveLinkHub/Private/LiveLinkHubInputProcessor.cpp
@@ -0,0 +1,75 @@
+// Copyright Epic Games, Inc. All Rights Reserved.
+#pragma once
+
+#include "LiveLinkHubInputProcessor.h"
+#include "Input/Events.h"
+#include "Misc/App.h"
+#include "ViewportWorldInteraction.h"
+
+namespace UE::LiveLinkHub::Private
+{
+	bool IsControllerButton( FName Name )
+	{
+		static FGamepadKeyNames::Type GamepadTypes[] = {
+			FGamepadKeyNames::LeftAnalogX,
+			FGamepadKeyNames::LeftAnalogY,
+			FGamepadKeyNames::RightAnalogX,
+			FGamepadKeyNames::RightAnalogY,
+			FGamepadKeyNames::LeftTriggerAnalog,
+			FGamepadKeyNames::RightTriggerAnalog,
+			FGamepadKeyNames::LeftThumb,
+			FGamepadKeyNames::RightThumb,
+			FGamepadKeyNames::SpecialLeft,
+			FGamepadKeyNames::SpecialLeft_X,
+			FGamepadKeyNames::SpecialLeft_Y,
+			FGamepadKeyNames::SpecialRight,
+			FGamepadKeyNames::FaceButtonBottom,
+			FGamepadKeyNames::FaceButtonRight,
+			FGamepadKeyNames::FaceButtonLeft,
+			FGamepadKeyNames::FaceButtonTop,
+			FGamepadKeyNames::LeftShoulder,
+			FGamepadKeyNames::RightShoulder,
+			FGamepadKeyNames::LeftTriggerThreshold,
+			FGamepadKeyNames::RightTriggerThreshold,
+			FGamepadKeyNames::DPadUp,
+			FGamepadKeyNames::DPadDown,
+			FGamepadKeyNames::DPadRight,
+			FGamepadKeyNames::DPadLeft,
+			FGamepadKeyNames::LeftStickUp,
+			FGamepadKeyNames::LeftStickDown,
+			FGamepadKeyNames::LeftStickRight,
+			FGamepadKeyNames::LeftStickLeft,
+			FGamepadKeyNames::RightStickUp,
+			FGamepadKeyNames::RightStickDown,
+			FGamepadKeyNames::RightStickRight,
+			FGamepadKeyNames::RightStickLeft
+		};
+
+		int32 Index = 0;
+		const int32 NumberOfTypes = sizeof(GamepadTypes)/sizeof(FGamepadKeyNames::Type);
+		bool bIsGamepadType = false;
+		while (Index < NumberOfTypes && !bIsGamepadType)
+		{
+			bIsGamepadType = GamepadTypes[Index] == Name;
+			Index++;
+		}
+
+		return bIsGamepadType;
+	}
+}
+
+bool FLiveLinkHubInputProcessor::HandleKeyDownEvent( FSlateApplication& SlateApp, const FKeyEvent& InKeyEvent )
+{
+	return UE::LiveLinkHub::Private::IsControllerButton(InKeyEvent.GetKey().GetFName());
+}
+
+bool FLiveLinkHubInputProcessor::HandleKeyUpEvent( FSlateApplication& SlateApp, const FKeyEvent& InKeyEvent )
+{
+	return UE::LiveLinkHub::Private::IsControllerButton(InKeyEvent.GetKey().GetFName());
+}
+
+bool FLiveLinkHubInputProcessor::HandleAnalogInputEvent( FSlateApplication& SlateApp, const FAnalogInputEvent& InAnalogInputEvent )
+{
+	return true;
+}
+
diff --git a/Engine/Plugins/Animation/LiveLinkHub/Source/LiveLinkHub/Private/LiveLinkHubInputProcessor.h b/Engine/Plugins/Animation/LiveLinkHub/Source/LiveLinkHub/Private/LiveLinkHubInputProcessor.h
new file mode 100644
--- /dev/null
+++ b/Engine/Plugins/Animation/LiveLinkHub/Source/LiveLinkHub/Private/LiveLinkHubInputProcessor.h
@@ -0,0 +1,28 @@
+// Copyright Epic Games, Inc. All Rights Reserved.
+#pragma once
+
+#include "GenericPlatform/ICursor.h"
+#include "Framework/Application/IInputProcessor.h"
+
+class FSlateApplication;
+class FViewportWorldInteractionManager;
+struct FAnalogInputEvent;
+struct FKeyEvent;
+struct FPointerEvent;
+
+/** This class will allow us to suppress gamepad input events for the slate application.*/
+class FLiveLinkHubInputProcessor : public IInputProcessor
+{
+public:
+	FLiveLinkHubInputProcessor() = default;
+	virtual ~FLiveLinkHubInputProcessor() = default;
+
+	//~ IInputProcess overrides
+	virtual bool HandleKeyDownEvent( FSlateApplication& SlateApp, const FKeyEvent& InKeyEvent ) override;
+	virtual bool HandleKeyUpEvent( FSlateApplication& SlateApp, const FKeyEvent& InKeyEvent ) override;
+	virtual bool HandleAnalogInputEvent( FSlateApplication& SlateApp, const FAnalogInputEvent& InAnalogInputEvent ) override;
+
+	virtual void Tick(const float DeltaTime, FSlateApplication& SlateApp, TSharedRef<ICursor> Cursor) override {}
+
+	virtual const TCHAR* GetDebugName() const override { return TEXT("LiveLinkHubInputProcessor"); }
+};
diff --git a/Engine/Plugins/VirtualProduction/LiveLinkInputDevice/Source/LiveLinkInputDevice/Private/LiveLinkInputDeviceSource.cpp b/Engine/Plugins/VirtualProduction/LiveLinkInputDevice/Source/LiveLinkInputDevice/Private/LiveLinkInputDeviceSource.cpp
--- a/Engine/Plugins/VirtualProduction/LiveLinkInputDevice/Source/LiveLinkInputDevice/Private/LiveLinkInputDeviceSource.cpp
+++ b/Engine/Plugins/VirtualProduction/LiveLinkInputDevice/Source/LiveLinkInputDevice/Private/LiveLinkInputDeviceSource.cpp
@@ -8,6 +8,7 @@
 #include "Features/IModularFeature.h"
 #include "GenericPlatform/GenericApplicationMessageHandler.h"
 #include "GenericPlatform/IInputInterface.h"
+#include "GenericPlatform/GenericPlatformInputDeviceMapper.h"
 #include "IInputDeviceModule.h"
 #include "InputCoreTypes.h"
 #include "Misc/CoreDelegates.h"
@@ -31,6 +32,7 @@ FLiveLinkInputDeviceSource::FLiveLinkInputDeviceSource(const FLiveLinkInputDevic
 	SourceMachineName = LOCTEXT("LocalMachine", "Local");
 
 	DeferredStartDelegateHandle = FCoreDelegates::OnEndFrame.AddRaw(this, &FLiveLinkInputDeviceSource::Start);
+	IPlatformInputDeviceMapper::Get().GetOnInputDeviceConnectionChange().AddRaw(this, &FLiveLinkInputDeviceSource::HandleGamepadDeviceChange);
 }
 
 FLiveLinkInputDeviceSource::~FLiveLinkInputDeviceSource()
@@ -39,6 +41,7 @@ FLiveLinkInputDeviceSource::~FLiveLinkInputDeviceSource()
 	{
 		FCoreDelegates::OnEndFrame.Remove(DeferredStartDelegateHandle);
 	}
+	IPlatformInputDeviceMapper::Get().GetOnInputDeviceConnectionChange().RemoveAll(this);
 }
 
 void FLiveLinkInputDeviceSource::ReceiveClient(ILiveLinkClient* InClient, FGuid InSourceGuid)
@@ -126,23 +129,24 @@ bool FLiveLinkInputDeviceSource::TickInputDevices(const double InDeltaTime)
 {
 	static thread_local double ElapsedUpdate = 0;
 	bool bGamepadState = false;
+
+	const bool bShouldDoDeviceUpdate = bShouldRefreshDeviceList;
 	for( int Index = 0; Index < InputDevices.Num(); Index++ )
 	{
-		InputDevices[Index]->Tick( InDeltaTime );
-		InputDevices[Index]->SendControllerEvents();
-		const bool bDeviceGamepadAttached = InputDevices[Index]->IsGamepadAttached();
-		// If elapsed time is greater than 1 second.  Force a device update so that we can catch new devices
-		// added to the system.  This normally happens by the WindowsApplication.cpp but we are outside of the
-		// process message loop so we have to periodically apply it here. We only do this here if we don't have
-		// any devices attached as there is a cost for making this query that can cause noticeble hitches.
-		//
-		if (!bDeviceGamepadAttached && ElapsedUpdate > 1.0)
+		if (bShouldDoDeviceUpdate)
 		{
 			static const FInputDeviceProperty RequestUpdateProp(TEXT("Request_Device_Update"));
 			InputDevices[Index]->SetDeviceProperty(-1, &RequestUpdateProp);
 		}
+		InputDevices[Index]->Tick( InDeltaTime );
+		InputDevices[Index]->SendControllerEvents();
+		const bool bDeviceGamepadAttached = InputDevices[Index]->IsGamepadAttached();
 		bGamepadState = bGamepadState || bDeviceGamepadAttached;
 	}
+	if (bShouldDoDeviceUpdate)
+	{
+		bShouldRefreshDeviceList = false;
+	}
 
 	ElapsedUpdate = ElapsedUpdate > 1 ? 0 : ElapsedUpdate + InDeltaTime;
 
@@ -280,6 +284,10 @@ void FLiveLinkInputDeviceSource::Send(FInputDeviceId DeviceId, FLiveLinkFrameDat
 	Client->PushSubjectFrameData_AnyThread({ SourceGuid, SubjectName }, MoveTemp(FrameDataToSend));
 }
 
+void FLiveLinkInputDeviceSource::HandleGamepadDeviceChange(EInputDeviceConnectionState /* unused */, FPlatformUserId /* unused  */, FInputDeviceId /* unused */)
+{
+    bShouldRefreshDeviceList = true;
+}
 
 
 #undef LOCTEXT_NAMESPACE
diff --git a/Engine/Plugins/VirtualProduction/LiveLinkInputDevice/Source/LiveLinkInputDevice/Public/LiveLinkInputDeviceSource.h b/Engine/Plugins/VirtualProduction/LiveLinkInputDevice/Source/LiveLinkInputDevice/Public/LiveLinkInputDeviceSource.h
--- a/Engine/Plugins/VirtualProduction/LiveLinkInputDevice/Source/LiveLinkInputDevice/Public/LiveLinkInputDeviceSource.h
+++ b/Engine/Plugins/VirtualProduction/LiveLinkInputDevice/Source/LiveLinkInputDevice/Public/LiveLinkInputDeviceSource.h
@@ -67,6 +67,9 @@ private:
 	double DevicePollWaitTime() const;
 
 private:
+	/** This delegate is invoked when a gamepad device is added or removed from the application. */
+	void HandleGamepadDeviceChange(EInputDeviceConnectionState NewConnectionState, FPlatformUserId PlatformUserId, FInputDeviceId InputDeviceId);
+
 	/** Update the our known gamepad state so that it can be reflected back to the LiveLink UI. */
 	void UpdateGamepadState(const bool bInState);
 
@@ -112,4 +115,6 @@ private:
 
 	// Indicates if we have ANY gamepads connected and thus are an active device for Live Link.
 	bool bIsGamepadConnected = false;
+
+	std::atomic<bool> bShouldRefreshDeviceList = false;
 };
