Add SceneTime to the frame data of LiveLinkInputDevice.  This missing information was preventing input device sources from being recorded in a take recorder track.

diff --git a/Engine/Plugins/VirtualProduction/LiveLinkInputDevice/Source/LiveLinkInputDevice/Private/LiveLinkInputDeviceSource.cpp b/Engine/Plugins/VirtualProduction/LiveLinkInputDevice/Source/LiveLinkInputDevice/Private/LiveLinkInputDeviceSource.cpp
--- a/Engine/Plugins/VirtualProduction/LiveLinkInputDevice/Source/LiveLinkInputDevice/Private/LiveLinkInputDeviceSource.cpp
+++ b/Engine/Plugins/VirtualProduction/LiveLinkInputDevice/Source/LiveLinkInputDevice/Private/LiveLinkInputDeviceSource.cpp
@@ -11,6 +11,7 @@
 #include "GenericPlatform/GenericPlatformInputDeviceMapper.h"
 #include "IInputDeviceModule.h"
 #include "InputCoreTypes.h"
+#include "Misc/App.h"
 #include "Misc/CoreDelegates.h"
 #include "Roles/LiveLinkTransformRole.h"
 #include "Roles/LiveLinkInputDeviceTypes.h"
@@ -253,6 +254,11 @@ TOptional<FLiveLinkFrameDataStruct> FLiveLinkInputDeviceSource::PollGamepadFrame
 
 		*FrameData = MessageHandler->GetLatestValue({});
 		FrameData->WorldTime = FPlatformTime::Seconds();
+		const TOptional<FQualifiedFrameTime> CurrentFrameTime = FApp::GetCurrentFrameTime();
+		if (CurrentFrameTime)
+		{
+			FrameData->MetaData.SceneTime = *CurrentFrameTime;
+		}
 
 		return BaseFrameData;
 	}
